[UBA]
* Refactored out uba horde integration to its own module called UbaCoordinatorHorde. This required some changes like callback registration to register helpers to uba host
* Added UbaCoordinatorHorde target in uba to be able to compile a coordinator for horde that can be used in UbaCli or other custom written uba tools

diff --git a/Engine/Plugins/UbaController/Source/UbaController/Private/UbaHordeAgent.cpp b/Engine/Plugins/UbaController/Source/UbaController/Private/UbaHordeAgent.cpp
deleted file mode 100644
diff --git a/Engine/Plugins/UbaController/Source/UbaController/Private/UbaHordeAgentManager.cpp b/Engine/Plugins/UbaController/Source/UbaController/Private/UbaHordeAgentManager.cpp
deleted file mode 100644
diff --git a/Engine/Plugins/UbaController/Source/UbaController/Private/UbaHordeComputeTransport.cpp b/Engine/Plugins/UbaController/Source/UbaController/Private/UbaHordeComputeTransport.cpp
deleted file mode 100644
diff --git a/Engine/Plugins/UbaController/Source/UbaController/Private/UbaHordeMetaClient.cpp b/Engine/Plugins/UbaController/Source/UbaController/Private/UbaHordeMetaClient.cpp
deleted file mode 100644
diff --git a/Engine/Plugins/UbaController/Source/UbaController/Private/UbaJobProcessor.cpp b/Engine/Plugins/UbaController/Source/UbaController/Private/UbaJobProcessor.cpp
--- a/Engine/Plugins/UbaController/Source/UbaController/Private/UbaJobProcessor.cpp
+++ b/Engine/Plugins/UbaController/Source/UbaController/Private/UbaJobProcessor.cpp
@@ -5,6 +5,7 @@
 #include "Interfaces/ITargetPlatform.h"
 #include "Interfaces/ITargetPlatformManagerModule.h"
 #include "Misc/ConfigCacheIni.h"
+#include "Misc/CoreMisc.h"
 #include "UbaHordeAgentManager.h"
 #include "Windows/AllowWindowsPlatformTypes.h"
 #include "UbaControllerModule.h"
@@ -287,6 +288,8 @@ void FUbaJobProcessor::RunTaskWithUba(FTask* Task)
 	Scheduler_EnqueueProcess(UbaScheduler, ProcessInfo, 1.0f, KnownInputsBuffer.GetData(), KnownInputsBuffer.Num()*sizeof(uba::tchar), KnownInputsCount);
 }
 
+FString GetUbaBinariesPath();
+
 void FUbaJobProcessor::StartUba()
 {
 	checkf(UbaServer == nullptr, TEXT("FUbaJobProcessor::StartUba() was called twice before FUbaJobProcessor::ShutDownUba()"));
@@ -333,7 +336,51 @@ void FUbaJobProcessor::StartUba()
 		Server_StartListen(UbaServer, uba::DefaultPort, nullptr); // Start listen so any helper on the LAN can join in
 	}
 
-	HordeAgentManager = MakeUnique<FUbaHordeAgentManager>(ControllerModule.GetWorkingDirectory(), UbaServer);
+	HordeAgentManager = MakeUnique<FUbaHordeAgentManager>(ControllerModule.GetWorkingDirectory(), GetUbaBinariesPath());
+
+	auto AddClientCallback = [](void* userData, const uba::tchar* ip, uint16 port)
+		{
+			return Server_AddClient((uba::NetworkServer*)userData, ip, port, nullptr);
+		};
+	HordeAgentManager->SetAddClientCallback(AddClientCallback, UbaServer);
+
+	FString HordeConfig;
+	if (GConfig->GetString(TEXT("UbaController"), TEXT("Horde"), HordeConfig, GEngineIni))
+	{
+		HordeConfig.TrimStartInline();
+		HordeConfig.TrimEndInline();
+		HordeConfig.RemoveFromStart(TEXT("("));
+		HordeConfig.RemoveFromEnd(TEXT(")"));
+
+		FString Url;
+		if (FParse::Value(*HordeConfig, TEXT("Url="), Url))
+		{
+			UE_LOG(LogUbaController, Log, TEXT("Found UBA controller Url: \"%s\""), *Url);
+			HordeAgentManager->SetUrl(Url);
+		}
+
+		FString Pool;
+		if (FParse::Value(*HordeConfig, TEXT("Pool="), Pool))
+		{
+			UE_LOG(LogUbaController, Log, TEXT("Found UBA controller Pool: \"%s\""), *Pool);
+			HordeAgentManager->SetPool(Pool);
+		}
+
+		FString Oidc;
+		if (FParse::Value(*HordeConfig, TEXT("Oidc="), Oidc))
+		{
+			UE_LOG(LogUbaController, Log, TEXT("Found UBA controller Oidc: \"%s\""), *Oidc);
+			HordeAgentManager->SetOidc(Oidc);
+		}
+
+		uint32 MaxCores = 0;
+		if (FParse::Value(*HordeConfig, TEXT("MaxCores="), MaxCores))
+		{
+			UE_LOG(LogUbaController, Log, TEXT("Found UBA controller MaxCores: \"%u\""), MaxCores);
+			HordeAgentManager->SetMaxCoreCount(MaxCores);
+		}
+	}
+
 
 	UE_LOG(LogUbaController, Display, TEXT("Created UBA storage server: RootDir=%s"), *RootDir);
 }
diff --git a/Engine/Plugins/UbaController/Source/UbaController/Public/UbaHordeAgent.h b/Engine/Plugins/UbaController/Source/UbaController/Public/UbaHordeAgent.h
deleted file mode 100644
diff --git a/Engine/Plugins/UbaController/Source/UbaController/Public/UbaHordeAgentManager.h b/Engine/Plugins/UbaController/Source/UbaController/Public/UbaHordeAgentManager.h
deleted file mode 100644
diff --git a/Engine/Plugins/UbaController/Source/UbaController/Public/UbaHordeComputeTransport.h b/Engine/Plugins/UbaController/Source/UbaController/Public/UbaHordeComputeTransport.h
deleted file mode 100644
diff --git a/Engine/Plugins/UbaController/Source/UbaController/Public/UbaHordeMetaClient.h b/Engine/Plugins/UbaController/Source/UbaController/Public/UbaHordeMetaClient.h
deleted file mode 100644
diff --git a/Engine/Plugins/UbaController/Source/UbaController/UbaController.build.cs b/Engine/Plugins/UbaController/Source/UbaController/UbaController.build.cs
--- a/Engine/Plugins/UbaController/Source/UbaController/UbaController.build.cs
+++ b/Engine/Plugins/UbaController/Source/UbaController/UbaController.build.cs
@@ -17,6 +17,7 @@ public class UbaController : ModuleRules
 			"Horde",
 			"HTTP",
 			"Sockets",
+			"UbaCoordinatorHorde",
 			"Json"
 		});
 
diff --git a/Engine/Source/Developer/DesktopPlatform/DesktopPlatform.Build.cs b/Engine/Source/Developer/DesktopPlatform/DesktopPlatform.Build.cs
--- a/Engine/Source/Developer/DesktopPlatform/DesktopPlatform.Build.cs
+++ b/Engine/Source/Developer/DesktopPlatform/DesktopPlatform.Build.cs
@@ -23,15 +23,17 @@ public class DesktopPlatform : ModuleRules
 			}
 		);
 		
-
 		if (Target.IsInPlatformGroup(UnrealPlatformGroup.Linux))
 		{
 			if(Target.Type == TargetType.Editor)
 			{
 				DynamicallyLoadedModuleNames.Add("SlateFontDialog");
 			}
-			
-			DynamicallyLoadedModuleNames.Add("SlateFileDialogs");
+
+			if (Target.bCompileAgainstCoreUObject && Target.bUsesSlate)
+			{
+				DynamicallyLoadedModuleNames.Add("SlateFileDialogs");
+			}
 
 			AddEngineThirdPartyPrivateStaticDependencies(Target, "SDL2");
 		}
diff --git a/Engine/Source/Developer/Horde/Private/Compute/ComputePlatform.cpp b/Engine/Source/Developer/Horde/Private/Compute/ComputePlatform.cpp
--- a/Engine/Source/Developer/Horde/Private/Compute/ComputePlatform.cpp
+++ b/Engine/Source/Developer/Horde/Private/Compute/ComputePlatform.cpp
@@ -317,7 +317,7 @@ void FComputeMemoryMappedFile::Close()
 	}
 
 	int Fd = (int)(size_t)Handle;
-	if(Fd >= 0)
+	if(Fd > 0)
 	{
 		close(Fd);
 		Handle = nullptr;
diff --git a/Engine/Source/Developer/UbaCoordinatorHorde/Private/UbaCoordinatorHorde.cpp b/Engine/Source/Developer/UbaCoordinatorHorde/Private/UbaCoordinatorHorde.cpp
new file mode 100644
--- /dev/null
+++ b/Engine/Source/Developer/UbaCoordinatorHorde/Private/UbaCoordinatorHorde.cpp
@@ -0,0 +1,100 @@
+// Copyright Epic Games, Inc. All Rights Reserved.
+
+#include "Modules/ModuleManager.h"
+
+#if defined(UBA_COORDINATOR_HORDE_DLL)
+
+#include "HAL/Platform.h"
+#include "HAL/PlatformApplicationMisc.h"
+#include "HttpManager.h"
+#include "HttpModule.h"
+#include "Misc/CommandLine.h"
+#include "Misc/ConfigCacheIni.h"
+#include "Misc/OutputDeviceConsole.h"
+#include "Misc/OutputDeviceRedirector.h"
+#include "UbaCoordinator.h"
+#include "UbaHordeAgentManager.h"
+
+#if PLATFORM_WINDOWS
+#define UBA_EXPORT __declspec(dllexport) 
+#else
+#define UBA_EXPORT __attribute__ ((visibility("default")))
+#endif
+
+TCHAR GInternalProjectName[64] = TEXT( "UbaCoordinatorHorde" );
+const TCHAR *GForeignEngineDir = nullptr;
+
+class CoordinatorHorde : public uba::Coordinator
+{
+public:
+	CoordinatorHorde(const FString& workDir, const FString& binariesDir) : m_manager(workDir, binariesDir)
+	{
+	}
+
+	virtual ~CoordinatorHorde()
+	{
+	}
+
+	virtual void SetAddClientCallback(AddClientCallback callback, void* userData) override
+	{
+		m_manager.SetAddClientCallback(callback, userData);
+	}
+
+	virtual void SetTargetCoreCount(uba::u32 count) override
+	{
+		m_manager.SetTargetCoreCount(count);
+	}
+
+	FUbaHordeAgentManager m_manager;
+};
+
+
+extern "C"
+{
+	UBA_EXPORT uba::Coordinator* UbaCreateCoordinator(const uba::CoordinatorCreateInfo& info)
+	{
+		FCommandLine::Set(TEXT(""));
+		if (info.logging)
+		{
+			FConfigCacheIni::InitializeConfigSystem();
+			static auto c = TUniquePtr<FOutputDeviceConsole>(FPlatformApplicationMisc::CreateConsoleOutputDevice());
+			GLogConsole = c.Get();
+			GLogConsole->Show(true);
+			GLog->SetCurrentThreadAsPrimaryThread();
+			GLog->TryStartDedicatedPrimaryThread();
+			GLog->AddOutputDevice(GLogConsole);
+		}
+
+		auto coordinator = new CoordinatorHorde(info.workDir, info.binariesDir);
+		auto& m = coordinator->m_manager;
+		m.SetUrl(info.uri);
+		m.SetPool(info.pool);
+		m.SetOidc(info.oidc);
+		m.SetMaxCoreCount(info.maxCoreCount);
+
+		return coordinator;
+	}
+
+	UBA_EXPORT void UbaDestroyCoordinator(uba::Coordinator* coordinator)
+	{
+		delete (CoordinatorHorde*)coordinator;
+		if (GLog)
+			GLog->SetCurrentThreadAsPrimaryThread();
+		FTextKey::TearDown(); // This is for clean shutdown with tsan
+		FHttpModule::Get().GetHttpManager().Shutdown();
+	}
+}
+
+#else
+
+class FUbaCoordinatorHordeModule : public IModuleInterface
+{
+	UBACOORDINATORHORDE_API FUbaCoordinatorHordeModule& Get()
+	{
+		return FModuleManager::LoadModuleChecked<FUbaCoordinatorHordeModule>(TEXT("UbaCoordinatorHorde"));
+	}
+};
+
+IMPLEMENT_MODULE(FUbaCoordinatorHordeModule, UbaCoordinatorHorde);
+
+#endif
\ No newline at end of file
diff --git a/Engine/Source/Developer/UbaCoordinatorHorde/Private/UbaHordeAgent.cpp b/Engine/Source/Developer/UbaCoordinatorHorde/Private/UbaHordeAgent.cpp
new file mode 100644
--- /dev/null
+++ b/Engine/Source/Developer/UbaCoordinatorHorde/Private/UbaHordeAgent.cpp
@@ -0,0 +1,247 @@
+// Copyright Epic Games, Inc. All Rights Reserved.
+
+#include "UbaHordeAgent.h"
+
+#include "Misc/Paths.h"
+#include "HAL/FileManager.h"
+#include "Compute/AgentMessage.h"
+#include "Compute/ComputeSocket.h"
+#include "Compute/ComputeTransport.h"
+#include "UbaHordeComputeTransport.h"
+
+DEFINE_LOG_CATEGORY(LogUbaHordeAgent);
+
+FUbaHordeAgent::FUbaHordeAgent(const FHordeRemoteMachineInfo& InMachineInfo)
+	: bIsValid(false)
+	, bHasErrors(false)
+	, MachineInfo(InMachineInfo)
+{
+	// Create the compute transport object and directly send over the nonce.
+	TUniquePtr<FUbaHordeComputeTransport> ComputeTransport = MakeUnique<FUbaHordeComputeTransport>(MachineInfo, bHasErrors);
+	if (!ComputeTransport->IsValid())
+	{
+		return;
+	}
+
+	ComputeTransport->Send(MachineInfo.Nonce, sizeof(MachineInfo.Nonce));
+
+	// Create the compute socket and initialize a compute channel with the recv/send buffers.
+	// This will allow us to create the agent channel object which will allow us to communicate
+	// with the agent directly.
+	HordeComputeSocket = CreateComputeSocket(MoveTemp(ComputeTransport), EComputeSocketEndpoint::Remote);
+	TSharedPtr<FComputeChannel> ComputeChannel = HordeComputeSocket->CreateChannel(0);
+	TSharedPtr<FComputeChannel> SecondComputeChannel = HordeComputeSocket->CreateChannel(100);
+
+	// Now this is what is going to handle the ComputeChannel.
+	AgentChannel = MakeUnique<FAgentMessageChannel>(MoveTemp(ComputeChannel));
+	ChildChannel = MakeUnique<FAgentMessageChannel>(MoveTemp(SecondComputeChannel));
+
+	bIsValid = true;
+}
+
+FUbaHordeAgent::~FUbaHordeAgent()
+{
+}
+
+bool FUbaHordeAgent::BeginCommunication()
+{
+	ensure(bIsValid);
+
+	// Now start the communication. From now, we can send and receive stuff.
+	HordeComputeSocket->StartCommunication();
+
+	// Let's try to read the response from the server. It should be an Attach message.
+	EAgentMessageType Type = AgentChannel->ReadResponse(5000);
+	UE_LOG(LogUbaHordeAgent, Verbose, TEXT("Got a response from the server: Type=0x%08X"), (int)Type);
+
+	if (Type == EAgentMessageType::None) // Timed out
+	{
+		return false;
+	}
+
+	if (!ensure(Type == EAgentMessageType::Attach))
+	{
+		return false;
+	}
+
+	/* Fork communication channel to split communication between a main and child channels. */
+	AgentChannel->Fork(100, 4*1024*1024);
+
+	Type = ChildChannel->ReadResponse(5000);
+
+	if (Type == EAgentMessageType::None) // Timed out
+	{
+		return false;
+	}
+
+	if (!ensure(Type == EAgentMessageType::Attach))
+	{
+		return false;
+	}
+
+	return true;
+}
+
+bool FUbaHordeAgent::UploadBinaries(const FString& BundleDirectory, const char* BundleLocator)
+{
+	// Prepare for giving the remote machine work.
+	// 1) Send binaries to remote machine.
+	// 2) Request to start executing UbaAgent.exe.
+	// 3) From there, it's a simple matter of sending work over through UBA.
+	ChildChannel->UploadFiles("", BundleLocator);
+
+	TMap<FString, FArchive*> BlobFiles;
+
+	auto FindOrAddBlobFile = [&BlobFiles, &BundleDirectory](const AgentMessage::FBlobRequest& BlobRequest) -> FArchive*
+		{
+			const FString Locator(BlobRequest.Locator.Len(), BlobRequest.Locator.GetData());
+
+			if (Locator.IsEmpty())
+			{
+				UE_LOG(LogUbaHordeAgent, Error, TEXT("Cannot upload binaries to Horde agent with empty locator"));
+				return nullptr;
+			}
+
+			if (FArchive** ArchiveElem = BlobFiles.Find(Locator))
+			{
+				return *ArchiveElem;
+			}
+			else
+			{
+				FString Path = FPaths::Combine(BundleDirectory, Locator);
+				Path.Append(TEXT(".blob"));
+
+				if (FArchive* Archive = IFileManager::Get().CreateFileReader(*Path))
+				{
+					BlobFiles.Add({ Locator, Archive });
+					return Archive;
+				}
+				else
+				{
+					UE_LOG(LogUbaHordeAgent, Error, TEXT("Cannot read blob file for Horde agent upload: %s"), *Path);
+					return nullptr;
+				}
+			}
+		};
+
+	EAgentMessageType Type = EAgentMessageType::None;
+
+	for (;;)
+	{
+		Type = ChildChannel->ReadResponse();
+		if (Type != EAgentMessageType::ReadBlob)
+		{
+			break;
+		}
+
+		AgentMessage::FBlobRequest BlobRequest;
+		ChildChannel->ReadBlobRequest(BlobRequest);
+
+		TArray<UTF8CHAR> LocatorData(BlobRequest.Locator.GetData(), BlobRequest.Locator.Len());
+		LocatorData.Add(UTF8CHAR('\0'));
+		UE_LOG(LogUbaHordeAgent, Verbose, TEXT("Response [ReadBlob]: Locator=%s, Offset=%d, Length=%d"), UTF8_TO_TCHAR(LocatorData.GetData()), BlobRequest.Offset, BlobRequest.Length);
+
+		FArchive* Archive = FindOrAddBlobFile(BlobRequest);
+		if (Archive == nullptr)
+		{
+			return false;
+		}
+
+		Archive->Seek(BlobRequest.Offset);
+
+		TArray<uint8> SerializedBytes;
+		SerializedBytes.SetNum(int(Archive->TotalSize() - BlobRequest.Offset));
+		Archive->Serialize(SerializedBytes.GetData(), SerializedBytes.Num());
+
+		// Send the blob data.
+		ChildChannel->Blob((unsigned char*)SerializedBytes.GetData(), SerializedBytes.Num());
+	}
+
+	// In case the agent upload was successful, WriteFilesResponse must have been the last received response.
+	// Otherwise, the remote machine might have abruptly shutdown, which can happen quite frequently on certain cloud services.
+	const bool bSuccess = (Type == EAgentMessageType::WriteFilesResponse);
+
+	return bSuccess;
+}
+
+void FUbaHordeAgent::Execute(const char* Exe, const char** Args, size_t NumArgs, const char* WorkingDir, const char** EnvVars, size_t NumEnvVars, bool bUseWine)
+{
+	ChildChannel->Execute(Exe, Args, NumArgs, WorkingDir, EnvVars, NumEnvVars, bUseWine ? EExecuteProcessFlags::UseWine : EExecuteProcessFlags::None);
+}
+
+void FUbaHordeAgent::CloseConnection()
+{
+	ChildChannel->Close();
+	AgentChannel->Close();
+}
+
+bool FUbaHordeAgent::IsValid()
+{
+	return bIsValid && !bHasErrors;
+}
+
+void FUbaHordeAgent::Poll(bool LogReports)
+{
+	EAgentMessageType Type = EAgentMessageType::None;
+	constexpr int32 ReadresponseTimeoutMS = 100;
+	while ((Type = ChildChannel->ReadResponse(ReadresponseTimeoutMS)) != EAgentMessageType::None)
+	{
+		switch (Type)
+		{
+		case EAgentMessageType::ExecuteOutput:
+		{
+			if (LogReports && ChildChannel->GetResponseSize() > 0)
+			{
+				// Convert raw buffer to dynamic array for modification
+				const ANSICHAR* ResponseDataRaw = reinterpret_cast<const ANSICHAR*>(ChildChannel->GetResponseData());
+				TArray<ANSICHAR> ResponseData(ResponseDataRaw, ChildChannel->GetResponseSize());
+
+				// Remove trailing newline characters and add NUL-terminator
+				while (ResponseData.Num() > 0 && (ResponseData.Last() == '\n' || ResponseData.Last() == '\r'))
+				{
+					ResponseData.Pop(EAllowShrinking::No);
+				}
+				ResponseData.Add('\0');
+
+				// Report output to log
+				auto ResponseString = StringCast<TCHAR>(ResponseData.GetData());
+				if (ResponseString.Length() > 0)
+				{
+					UE_LOG(LogUbaHordeAgent, Log, TEXT("Response [ExecuteOutput]: %s"), ResponseString.Get());
+				}
+			}
+		}
+		break;
+
+		case EAgentMessageType::ExecuteResult:
+		{
+			if (ChildChannel->GetResponseSize() == sizeof(int32))
+			{
+				if (LogReports)
+				{
+					const int32* ResponseData = reinterpret_cast<const int32*>(ChildChannel->GetResponseData());
+					const int32 ExecuteExitCode = ResponseData[0];
+					UE_LOG(LogUbaHorde, Log, TEXT("Response [ExecuteResult]: ExitCode=%d"), ExecuteExitCode);
+				}
+				bIsValid = false;
+			}
+		}
+		break;
+
+		case EAgentMessageType::Exception:
+		{
+			AgentMessage::FException Ex;
+			ChildChannel->ReadException(Ex);
+			UE_LOG(LogUbaHorde, Log, TEXT("EXCEPTION Message: %s"), *FString(Ex.Message));
+			UE_LOG(LogUbaHorde, Log, TEXT("EXCEPTION Description: %s"), *FString(Ex.Description));
+		}
+		break;
+
+		default:
+		{
+			// Ignore other messages here, we are only interested in the ExecuteOutput to forward reports from the remote agents to the calling process
+		}
+		break;
+		}
+	}
+}
diff --git a/Engine/Source/Developer/UbaCoordinatorHorde/Private/UbaHordeAgent.h b/Engine/Source/Developer/UbaCoordinatorHorde/Private/UbaHordeAgent.h
new file mode 100644
--- /dev/null
+++ b/Engine/Source/Developer/UbaCoordinatorHorde/Private/UbaHordeAgent.h
@@ -0,0 +1,50 @@
+// Copyright Epic Games, Inc. All Rights Reserved.
+#pragma once
+
+#include <memory>
+
+#include "UbaHordeMetaClient.h"
+
+class FComputeSocket;
+class FAgentMessageChannel;
+
+DECLARE_LOG_CATEGORY_EXTERN(LogUbaHordeAgent, Log, Log);
+
+class FUbaHordeAgent
+{
+public:
+	FUbaHordeAgent(const FHordeRemoteMachineInfo &MachineInfo);
+
+	~FUbaHordeAgent();
+
+	/* Start spinning the ping thread, and receive/send threads. */
+	bool BeginCommunication();
+
+	/* Upload the proper executables/binaries. The Locator is just the contents of the bundle ref file.
+	 * e.g. Locator = abcdef1234567890abcdef1234567890abcdef12@34567890abcdef1234567890#123 */
+	bool UploadBinaries(const FString& BundleDirectory, const char* BundleLocator);
+
+	void Execute(const char* Exe, const char** Args, size_t NumArgs, const char* WorkingDir, const char** EnvVars, size_t NumEnvVars, bool bUseWine = false);
+
+	void CloseConnection(); 
+
+	bool IsValid();
+
+	// Reads output from the child channel and reports them to the log of the calling process.
+	void Poll(bool LogReports);
+
+	inline const FHordeRemoteMachineInfo& GetMachineInfo() const
+	{
+		return MachineInfo;
+	}
+
+private:
+	/* For Horde integration. Once all these are initialized, we're good to using UBA */
+	TUniquePtr<FComputeSocket> HordeComputeSocket;
+	TUniquePtr<FAgentMessageChannel> AgentChannel;
+	TUniquePtr<FAgentMessageChannel> ChildChannel;
+
+	bool bIsValid;
+	bool bHasErrors;
+	FHordeRemoteMachineInfo MachineInfo;
+};
diff --git a/Engine/Source/Developer/UbaCoordinatorHorde/Private/UbaHordeAgentManager.cpp b/Engine/Source/Developer/UbaCoordinatorHorde/Private/UbaHordeAgentManager.cpp
new file mode 100644
--- /dev/null
+++ b/Engine/Source/Developer/UbaCoordinatorHorde/Private/UbaHordeAgentManager.cpp
@@ -0,0 +1,394 @@
+// Copyright Epic Games, Inc. All Rights Reserved.
+
+#include "UbaHordeAgentManager.h"
+#include "HAL/Event.h"
+#include "HAL/IConsoleManager.h"
+#include "HAL/PlatformTime.h"
+#include "Misc/ConfigCacheIni.h"
+#include "Misc/FileHelper.h"
+#include "Misc/Paths.h"
+#include "Misc/ScopeExit.h"
+#include "Storage/StorageClient.h"
+#include "Storage/Clients/BundleStorageClient.h"
+#include "Storage/Clients/MemoryStorageClient.h"
+#include "Storage/Clients/FileStorageClient.h"
+#include "Storage/Nodes/ChunkNode.h"
+#include "Storage/Nodes/DirectoryNode.h"
+#include "Storage/BlobWriter.h"
+#include "UbaBase.h"
+#include "UbaHordeAgent.h"
+#include <filesystem>
+#include <fstream>
+
+namespace UbaCoordinatorHordeModule
+{
+	static bool bHordeForwardAgentLogs = false;
+	static FAutoConsoleVariableRef CVarUbaControllerHordeForwardAgentLogs(
+		TEXT("r.UbaHorde.ForwardAgentLogs"),
+		bHordeForwardAgentLogs,
+		TEXT("Enables or disables logging of stdout on agent side to show in controller log."));
+}
+
+FUbaHordeAgentManager::FUbaHordeAgentManager(const FString& InWorkingDir, const FString& InBinariesPath)
+	:	WorkingDir(InWorkingDir)
+	,	BinariesPath(InBinariesPath)
+	,	LastRequestFailTime(1)
+	,	TargetCoreCount(0)
+	,	EstimatedCoreCount(0)
+	,   ActiveCoreCount(0)
+	,	AskForAgents(true)
+{
+}
+
+FUbaHordeAgentManager::~FUbaHordeAgentManager()
+{
+	FScopeLock AgentsScopeLock(&AgentsLock);
+	for (TUniquePtr<FHordeAgentWrapper>& Agent : Agents)
+	{
+		Agent->ShouldExit->Trigger();
+		Agent->Thread.Join();
+		FGenericPlatformProcess::ReturnSynchEventToPool(Agent->ShouldExit);
+	}
+}
+
+void FUbaHordeAgentManager::SetTargetCoreCount(uint32 Count)
+{
+	TargetCoreCount = FMath::Min(MaxCores, Count);
+
+	if (Url.IsEmpty())
+	{
+		AskForAgents = false;
+	}
+
+	while (EstimatedCoreCount < TargetCoreCount)
+	{
+		if (!AskForAgents)
+		{
+			return;
+		}
+
+		//UE_LOG(LogUbaHorde, Display, TEXT("Requested new agent. Estimated core count: %u, Target core count: %u"), EstimatedCoreCount.Load(), TargetCoreCount.Load());
+		RequestAgent();
+	}
+
+	FScopeLock AgentsScopeLock(&AgentsLock);
+	for (auto Iterator = Agents.CreateIterator(); Iterator; ++Iterator)
+	{
+		TUniquePtr<FHordeAgentWrapper>& Agent = *Iterator;
+		if (Agent->ShouldExit->Wait(0))
+		{
+			Agent->Thread.Join();
+			FGenericPlatformProcess::ReturnSynchEventToPool(Agent->ShouldExit);
+			Iterator.RemoveCurrentSwap();
+		}
+	}
+}
+
+void FUbaHordeAgentManager::SetAddClientCallback(AddClientCallback* callback, void* userData)
+{
+	m_callback = callback;
+	m_userData = userData;
+}
+
+int32 FUbaHordeAgentManager::GetAgentCount() const
+{
+	FScopeLock AgentsScopeLock(&AgentsLock);
+	return Agents.Num();
+}
+
+uint32 FUbaHordeAgentManager::GetActiveCoreCount() const
+{
+	return ActiveCoreCount;
+}
+
+// Creates a bundle blob (one of several chunks of a file) to be uploaded to Horde
+// This code has been adopted from the HordeTest project.
+// See 'Engine/Source/Programs/Horde/Samples/HordeTest/Main.cpp'.
+static FBlobHandleWithHash CreateHordeBundleBlob(const std::filesystem::path& Path, FBlobWriter& Writer, int64& OutLength, FIoHash& OutStreamHash)
+{
+	OutLength = 0;
+
+	FChunkNodeWriter ChunkWriter(Writer);
+
+	std::ifstream Stream(Path, std::ios::binary);
+
+	char ReadBuffer[4096];
+	while (!Stream.eof())
+	{
+		Stream.read(ReadBuffer, sizeof(ReadBuffer));
+
+		const int64 ReadSize = static_cast<int64>(Stream.gcount());
+		if (ReadSize == 0)
+		{
+			break;
+		}
+		OutLength += ReadSize;
+
+		ChunkWriter.Write(FMemoryView(ReadBuffer, ReadSize));
+	}
+
+	return ChunkWriter.Flush(OutStreamHash);
+}
+
+static FDirectoryEntry CreateHordeBundleDirectoryEntry(const std::filesystem::path& Path, FBlobWriter& Writer)
+{
+	FDirectoryNode DirectoryNode;
+
+	int64 BlobLength = 0;
+	FIoHash StreamHash;
+	FBlobHandleWithHash Target = CreateHordeBundleBlob(Path, Writer, BlobLength, StreamHash);
+
+	EFileEntryFlags Flags = EFileEntryFlags::Executable;
+	FFileEntry NewEntry(Target, FUtf8String(Path.filename().string().c_str()), Flags, BlobLength, StreamHash, FSharedBufferView());
+
+	const FUtf8String Name = NewEntry.Name;
+	const int64 Length = NewEntry.Length;
+	DirectoryNode.NameToFile.Add(Name, MoveTemp(NewEntry));
+
+	FBlobHandle DirectoryHandle = DirectoryNode.Write(Writer);
+
+	return FDirectoryEntry(DirectoryHandle, FIoHash(), FUtf8String(Path.filename().string().c_str()), Length);
+}
+
+bool CreateHordeBundleFromFile(const std::filesystem::path& InputFilename, const std::filesystem::path& OutputFilename)
+{
+	TSharedRef<FFileStorageClient> FileStorage = MakeShared<FFileStorageClient>(OutputFilename.parent_path());
+	TSharedRef<FBundleStorageClient> Storage = MakeShared<FBundleStorageClient>(FileStorage);
+
+	TUniquePtr<FBlobWriter> Writer = Storage->CreateWriter("");
+	FDirectoryEntry RootEntry = CreateHordeBundleDirectoryEntry(InputFilename, *Writer.Get());
+	Writer->Flush();
+
+	FFileStorageClient::WriteRefToFile(OutputFilename, RootEntry.Target->GetLocator());
+	return true;
+}
+
+void FUbaHordeAgentManager::RequestAgent()
+{
+	EstimatedCoreCount += 32; // We estimate a typical agent to have 32 cores
+
+	FScopeLock AgentsScopeLock(&AgentsLock);
+	FHordeAgentWrapper& Wrapper = *Agents.Emplace_GetRef(MakeUnique<FHordeAgentWrapper>());
+
+	Wrapper.ShouldExit = FGenericPlatformProcess::GetSynchEventFromPool(true);
+	Wrapper.Thread = FThread(TEXT("HordeAgent"), [this, WrapperPtr = &Wrapper]() { ThreadAgent(*WrapperPtr); });
+}
+
+void FUbaHordeAgentManager::ThreadAgent(FHordeAgentWrapper& Wrapper)
+{
+	FEvent& ShouldExit = *Wrapper.ShouldExit;
+	TUniquePtr<FUbaHordeAgent> Agent;
+	bool bSuccess = false;
+
+	ON_SCOPE_EXIT
+	{
+		if (Agent)
+		{
+			Agent->CloseConnection();
+		}
+
+		ShouldExit.Trigger();
+	};
+
+	int MachineCoreCount = 0;
+	uint32 ListenPort = 7001;
+
+	{
+		ON_SCOPE_EXIT{ EstimatedCoreCount -= 32; };
+
+		#if PLATFORM_WINDOWS
+		const char* AppName = "UbaAgent.exe";
+		#else
+		const char* AppName = "UbaAgent";
+		#endif
+
+		FScopeLock ScopeLock(&BundleRefPathsLock);
+		if (BundleRefPaths.IsEmpty())
+		{
+			struct BundleRec { const char* File; const TCHAR* BundleRef; };
+			BundleRec BundleRecs[] =
+			{
+				{ AppName, TEXT("UbaAgent.Bundle.ref") },
+#if PLATFORM_LINUX
+				{ "UbaAgent.debug", TEXT("UbaAgent.debug.Bundle.ref") },
+				{ "libclang_rt.tsan.so", TEXT("Tsan.Bundle.ref") },
+#endif
+			};
+
+			for (const BundleRec& Rec : BundleRecs)
+			{
+				const FString FilePath = FPaths::Combine(BinariesPath, ANSI_TO_TCHAR(Rec.File));
+				FString BundlePath = FPaths::Combine(WorkingDir, Rec.BundleRef);
+
+				if (!CreateHordeBundleFromFile(*FilePath, *BundlePath))
+				{
+					UE_LOG(LogUbaHorde, Error, TEXT("Failed to create Horde bundle for: %s"), *FilePath);
+					AskForAgents = false;
+					return;
+				}
+				UE_LOG(LogUbaHorde, Display, TEXT("Created Horde bundle for: %s"), *FilePath);
+				BundleRefPaths.Add(BundlePath);
+			}
+		}
+
+		if (!HordeMetaClient)
+		{
+			// Create Horde meta client right before we need it to make sure the CVar for the server URL has been read by now
+			HordeMetaClient = MakeUnique<FUbaHordeMetaClient>(Url, Oidc);
+			if (!HordeMetaClient->RefreshHttpClient())
+			{
+				UE_LOG(LogUbaHorde, Error, TEXT("Failed to create HttpClient for UbaAgent"));
+				AskForAgents = false;
+				return;
+			}
+		}
+
+		if (!AskForAgents)
+		{
+			return;
+		}
+
+		if (LastRequestFailTime == 0)
+		{
+			ScopeLock.Unlock();
+		}
+		else
+		{
+			// Try to reduce pressure on horde by not asking for machines more frequent than every 5 seconds if failed to retrieve last time
+			uint64 CurrentTime = FPlatformTime::Cycles64();
+			uint32 MsSinceLastFail = uint32(double(CurrentTime - LastRequestFailTime) * FPlatformTime::GetSecondsPerCycle() * 1000);
+			if (MsSinceLastFail < 5000)
+			{
+				if (ShouldExit.Wait(5000 - MsSinceLastFail))
+				{
+					return;
+				}
+			}
+		}
+
+		TSharedPtr<FUbaHordeMetaClient::HordeMachinePromise, ESPMode::ThreadSafe> Promise = HordeMetaClient->RequestMachine(Pool);
+		if (!Promise)
+		{
+			//UE_LOG(LogUbaHorde, Error, TEXT("Failed to create Horde bundle for UbaAgent executable: %s"), *UbaAgentFilePath);
+			return;
+		}
+		TFuture<TTuple<FHttpResponsePtr, FHordeRemoteMachineInfo>> Future = Promise->GetFuture();
+		Future.Wait();
+		FHordeRemoteMachineInfo MachineInfo = Future.Get().Value;
+
+		// If the machine couldn't be assigned, just ignore this agent slot
+		if (MachineInfo.Ip == TEXT(""))
+		{
+			if (!LastRequestFailTime)
+			{
+				UE_LOG(LogUbaHorde, Verbose, TEXT("No resources available in Horde. Will keep retrying until %u cores are used (Currently have %u)"), TargetCoreCount.Load(), ActiveCoreCount.Load());
+			}
+			LastRequestFailTime = FPlatformTime::Cycles64();
+			return;
+		}
+
+		LastRequestFailTime = 0;
+
+		ScopeLock.Unlock();
+
+		if (ShouldExit.Wait(0))
+		{
+			return;
+		}
+
+		Agent = MakeUnique<FUbaHordeAgent>(MachineInfo);
+
+		if (!Agent->BeginCommunication())
+		{
+			return;
+		}
+
+		for (const FString& Bundle : BundleRefPaths)
+		{
+			TArray<uint8> Locator;
+			if (!FFileHelper::LoadFileToArray(Locator, *Bundle))
+			{
+				UE_LOG(LogUbaHorde, Error, TEXT("Cannot launch Horde processes for UBA controller because bundle path could not be found: %s"), *Bundle);
+				return;
+			}
+			Locator.Add('\0');
+
+			FString BundleDirectory = FPaths::GetPath(Bundle);
+
+			if (ShouldExit.Wait(0))
+			{
+				return;
+			}
+
+			if (!Agent->UploadBinaries(BundleDirectory, reinterpret_cast<const char*>(Locator.GetData())))
+			{
+				return;
+			}
+		}
+
+		// Start the UBA Agent that will connect to us, requesting for work
+		const std::string ListenPortArg = "-listen=" + std::to_string(ListenPort);
+
+		const char* UbaAgentArgs[] =
+		{
+			ListenPortArg.c_str(),
+			"-nopoll",				// -nopoll recommended when running on remote Horde agents to make sure they exit after completion. Otherwise, it keeps running.
+			"-listenTimeout=5",		// Agent will wait 5 seconds for this thread to connect (Server_AddClient does the connect)
+			"-quiet",				// Skip all the agent logging that would be sent over to here
+			"-maxidle=15",			// After 15 seconds of idling agent will automatically disconnect
+			"-Dir=%UE_HORDE_SHARED_DIR%\\Uba",
+			"-Eventfile=%UE_HORDE_TERMINATION_SIGNAL_FILE%",
+		};
+
+		// If the machine does not run Windows, enable the compatibility layer Wine to run UbaAgent.exe on POSIX systems
+		#if PLATFORM_WINDOWS
+		const bool bRunsWindowsOS = Agent->GetMachineInfo().bRunsWindowOS;
+		const bool bUseWine = !bRunsWindowsOS;
+		#else
+		const bool bUseWine = false;
+		#endif
+
+		if (ShouldExit.Wait(0))
+		{
+			return;
+		}
+
+		Agent->Execute(AppName, UbaAgentArgs, UE_ARRAY_COUNT(UbaAgentArgs), nullptr, nullptr, 0, bUseWine);
+
+		// Log remote execution
+		FString UbaAgentCmdArgs = ANSI_TO_TCHAR(AppName);
+		for (const char* Arg : UbaAgentArgs)
+		{
+			UbaAgentCmdArgs += TEXT(" ");
+			UbaAgentCmdArgs += ANSI_TO_TCHAR(Arg);
+		}
+		UE_LOG(LogUbaHorde, Log, TEXT("Remote execution on Horde machine [%s:%d]: %s"), *Agent->GetMachineInfo().Ip, ListenPort, *UbaAgentCmdArgs);
+
+		MachineCoreCount = MachineInfo.LogicalCores;
+		EstimatedCoreCount += MachineCoreCount;
+		ActiveCoreCount += MachineCoreCount;
+	}
+
+	uint32 callCounter = 0; // TODO: This should react on the listen string instead of waiting for two text messages :)
+
+	while (Agent->IsValid() && !ShouldExit.Wait(100))
+	{
+		Agent->Poll(UbaCoordinatorHordeModule::bHordeForwardAgentLogs);
+
+		if (callCounter++ == 2)
+		{
+			// Add this machine as client to the remote agent
+			const FString& IpAddress = Agent->GetMachineInfo().Ip;
+			const bool bAddClientSuccess = 	m_callback(m_userData, StringCast<uba::tchar>(*IpAddress).Get(), (uint16)ListenPort);
+
+			if (!bAddClientSuccess)
+			{
+				UE_LOG(LogUbaHorde, Display, TEXT("Server_AddClient(%s:%d) failed"), *IpAddress, ListenPort);
+				return;
+			}
+		}
+	}
+
+	ActiveCoreCount -= MachineCoreCount;
+	EstimatedCoreCount -= MachineCoreCount;
+}
diff --git a/Engine/Source/Developer/UbaCoordinatorHorde/Private/UbaHordeComputeTransport.cpp b/Engine/Source/Developer/UbaCoordinatorHorde/Private/UbaHordeComputeTransport.cpp
new file mode 100644
--- /dev/null
+++ b/Engine/Source/Developer/UbaCoordinatorHorde/Private/UbaHordeComputeTransport.cpp
@@ -0,0 +1,92 @@
+// Copyright Epic Games, Inc. All Rights Reserved.
+
+#include "UbaHordeComputeTransport.h"
+#include "Sockets.h"
+#include "SocketSubsystem.h"
+
+FUbaHordeComputeTransport::FUbaHordeComputeTransport(const FHordeRemoteMachineInfo& MachineInfo, bool& InHasErrors)
+	: Socket(nullptr)
+	, bIsClosed(false)
+	, bHasErrors(InHasErrors)
+{
+	ISocketSubsystem* SocketSubsystem = ISocketSubsystem::Get(PLATFORM_SOCKETSUBSYSTEM);
+
+	Socket = SocketSubsystem->CreateSocket(TEXT("Stream"), TEXT(""));
+	TSharedPtr<FInternetAddr> InternetAddr;
+
+	TSharedPtr<FInternetAddr> Address = SocketSubsystem->GetAddressFromString(MachineInfo.Ip);
+	Address->SetPort(MachineInfo.Port);
+
+	if (!Socket->Connect(*Address))
+	{
+		UE_LOG(LogUbaHorde, Warning, TEXT("Failed to connect to Horde Agent"));
+
+		SocketSubsystem->DestroySocket(Socket);
+
+		Socket = nullptr;
+	}
+}
+
+FUbaHordeComputeTransport::~FUbaHordeComputeTransport()
+{
+	if (Socket)
+	{
+		ISocketSubsystem* SocketSubsystem = ISocketSubsystem::Get(PLATFORM_SOCKETSUBSYSTEM);
+		SocketSubsystem->DestroySocket(Socket);
+	}
+}
+
+// Sends data to the remote
+size_t FUbaHordeComputeTransport::Send(const void* Data, size_t Size)
+{
+	int32 NumBytesSent = 0;
+	if (!Socket->Send((const uint8*)Data, Size, NumBytesSent))
+	{
+		NumBytesSent = 0;
+		// We can't log this since the other side could have disconnected us (causing recv to fail). This happens often in the horde setup
+		//UE_LOG(LogUbaHorde, Error, TEXT("Failed to send data to the Horde Agent: %llu %s"), (uint64)Size, Size == 1 ? TEXT("byte") : TEXT("bytes"));
+		bHasErrors = true;
+		return 0;
+	}
+
+	UE_LOG(LogUbaHorde, Verbose, TEXT("Sent message to horde agent: %d %s"), NumBytesSent, NumBytesSent == 1 ? TEXT("byte") : TEXT("bytes"));
+
+	return NumBytesSent;
+}
+
+// Receives data from the remote
+size_t FUbaHordeComputeTransport::Recv(void* Data, size_t Size)
+{
+	int32 NumBytesRead = 0;
+	if (!Socket->Recv((uint8*)Data, Size, NumBytesRead))
+	{
+		NumBytesRead = 0;
+
+		// We can't log this since the other side could have disconnected us (causing recv to fail). This happens often in the horde setup
+		//if (!bIsClosed)
+		//	UE_LOG(LogUbaHorde, Error, TEXT("Failed to receive data from the Horde Agent: %llu %s"), (uint64)Size, Size == 1 ? TEXT("byte") : TEXT("bytes"));
+		return 0;
+	}
+
+	UE_LOG(LogUbaHorde, Verbose, TEXT("Received message from horde agent: %d %s"), NumBytesRead, NumBytesRead == 1 ? TEXT("byte") : TEXT("bytes"));
+
+	return NumBytesRead;
+}
+
+// Indicates to the remote that no more data will be sent.
+void FUbaHordeComputeTransport::MarkComplete()
+{
+}
+
+// Indicates that no more data will be sent or received, and that any blocking reads/writes should stop.
+void FUbaHordeComputeTransport::Close()
+{
+	bIsClosed = true;
+	//Socket->Close();
+	Socket->Shutdown(ESocketShutdownMode::ReadWrite);
+}
+
+bool FUbaHordeComputeTransport::IsValid()
+{
+	return Socket != nullptr && !bHasErrors && !bIsClosed;
+}
diff --git a/Engine/Source/Developer/UbaCoordinatorHorde/Private/UbaHordeComputeTransport.h b/Engine/Source/Developer/UbaCoordinatorHorde/Private/UbaHordeComputeTransport.h
new file mode 100644
--- /dev/null
+++ b/Engine/Source/Developer/UbaCoordinatorHorde/Private/UbaHordeComputeTransport.h
@@ -0,0 +1,34 @@
+// Copyright Epic Games, Inc. All Rights Reserved.
+#pragma once
+
+#include "Compute/ComputeTransport.h"
+#include "UbaHordeMetaClient.h"
+
+class FSocket;
+
+class FUbaHordeComputeTransport : public FComputeTransport
+{
+public:
+	// The constructor just performs a connect.
+	FUbaHordeComputeTransport(const FHordeRemoteMachineInfo &Info, bool& InHasErrors);
+	~FUbaHordeComputeTransport() override;
+	
+	// Sends data to the remote
+	size_t Send(const void* Data, size_t Size) override;
+
+	// Receives data from the remote
+	size_t Recv(void* Data, size_t Size) override;
+
+	// Indicates to the remote that no more data will be sent.
+	void MarkComplete() override;
+
+	// Indicates that no more data will be sent or received, and that any blocking reads/writes should stop.
+	void Close() override;
+
+	bool IsValid();
+
+private:
+	FSocket* Socket;
+	bool bIsClosed;
+	bool& bHasErrors;
+};
diff --git a/Engine/Source/Developer/UbaCoordinatorHorde/Private/UbaHordeMetaClient.cpp b/Engine/Source/Developer/UbaCoordinatorHorde/Private/UbaHordeMetaClient.cpp
new file mode 100644
--- /dev/null
+++ b/Engine/Source/Developer/UbaCoordinatorHorde/Private/UbaHordeMetaClient.cpp
@@ -0,0 +1,169 @@
+// Copyright Epic Games, Inc. All Rights Reserved.
+
+#include "UbaHordeMetaClient.h"
+#include "Serialization/JsonReader.h"
+#include "Serialization/JsonSerializer.h"
+#include "Misc/App.h"
+#include "Misc/ConfigCacheIni.h"
+
+DEFINE_LOG_CATEGORY(LogUbaHorde);
+
+FUbaHordeMetaClient::FUbaHordeMetaClient(const FStringView& HordeServerUrl, const FStringView& InOAuthProviderIdentifier)
+	: ServerUrl(HordeServerUrl)
+	, bConnectWithAuthentication(HordeServerUrl.StartsWith(TEXT("https:")))
+	, OAuthProviderIdentifier(InOAuthProviderIdentifier)
+{
+}
+
+FUbaHordeMetaClient::FUbaHordeMetaClient(const FStringView& HordeServerIp, uint16 HordeServerPort, bool bInConnectWithAuthentication, const FStringView& InOAuthProviderIdentifier)
+	: ServerUrl(FString::Format(TEXT("{0}://{1}:{2}"), {(bInConnectWithAuthentication ? TEXT("https") : TEXT("http")), HordeServerIp, HordeServerPort}))
+	, bConnectWithAuthentication(bInConnectWithAuthentication)
+	, OAuthProviderIdentifier(InOAuthProviderIdentifier)
+{
+}
+
+bool FUbaHordeMetaClient::RefreshHttpClient()
+{
+	// Try to connect to Horde with HTTP and v2 API
+	HttpClient = MakeUnique<FHordeHttpClient>(ServerUrl);
+
+	if (bConnectWithAuthentication)
+	{
+		check(!OAuthProviderIdentifier.IsEmpty());
+		if (FApp::IsUnattended())
+		{
+			UE_LOG(LogUbaHorde, Display, TEXT("Logging in to Horde server with environment variable UE_HORDE_TOKEN: %s"), *ServerUrl);
+			if (HttpClient->LoginWithEnvironmentVariable())
+			{
+				return true;
+			}
+
+			// Fallback to OIDC login as some commandlets might add -Unattended implicitly, but not everyone has UE_HORDE_TOKEN environment variable available
+			UE_LOG(LogUbaHorde, Display, TEXT("Login attempt with environment variable failed; trying to login with OIDC: %s"), *ServerUrl);
+		}
+		else
+		{
+			UE_LOG(LogUbaHorde, Display, TEXT("Logging in to Horde server with OIDC: %s"), *ServerUrl);
+		}
+
+		if (HttpClient->LoginWithOidc(*OAuthProviderIdentifier, FApp::IsUnattended()))
+		{
+			return true;
+		}
+
+		UE_LOG(LogUbaHorde, Warning, TEXT("Login to Horde server [%s] failed"), *ServerUrl);
+		return false;
+	}
+	return true;
+}
+
+TSharedPtr<FUbaHordeMetaClient::HordeMachinePromise, ESPMode::ThreadSafe> FUbaHordeMetaClient::RequestMachine(const FString& PoolId, const FString& Machine)
+{
+	TSharedPtr<HordeMachinePromise> Promise = MakeShared<HordeMachinePromise, ESPMode::ThreadSafe>();
+
+	const FString ResourcePath = FString::Format(TEXT("api/v2/compute/{0}"), { Machine });
+
+	FHttpRequestRef Request = HttpClient->CreateRequest(TEXT("POST"), *ResourcePath);
+
+	Request->SetDelegateThreadPolicy(EHttpRequestDelegateThreadPolicy::CompleteOnHttpThread);
+
+	// Post JSON document with constraints to acquire a Horde agent. Use pool of agents (e.g. "BoxLinux" or "BoxWin")
+	// and require exclusive access or UbaStorage will fail to initialize the next time we connect to the same machine.
+	const FString Body = FString::Format(TEXT("{\"requirements\":{\"pool\":\"{0}\",\"exclusive\":true}}"), { PoolId });
+	Request->SetContentAsString(Body);
+	Request->SetHeader(TEXT("Content-Type"), TEXT("application/json"));
+
+	UE_LOG(LogUbaHorde, Log, TEXT("Requesting Horde agent with JSON descriptor: '%s'"), *Body);
+
+	Request->OnProcessRequestComplete().BindLambda(
+		[this, Promise](FHttpRequestPtr /*Request*/, FHttpResponsePtr HttpResponse, bool bConnectedSuccessfully)
+		{
+			FHordeRemoteMachineInfo Info;
+			Info.Ip = TEXT("");
+			Info.Port = 0xFFFF;
+			Info.bRunsWindowOS = false;
+			FMemory::Memset(Info.Nonce, 0, sizeof(Info.Nonce));
+
+			if (!bConnectedSuccessfully || !HttpResponse.IsValid())
+			{
+				UE_LOG(LogUbaHorde, Verbose, TEXT("No response from Horde"));
+				Promise->SetValue(MakeTuple(HttpResponse, Info));
+				return;
+			}
+
+			FString ResponseStr = HttpResponse->GetContentAsString();
+
+			if (HttpResponse->GetResponseCode() == 503)
+			{
+				UE_LOG(LogUbaHorde, Verbose, TEXT("No resources available in Horde (%s)"), *ResponseStr);
+				Promise->SetValue(MakeTuple(HttpResponse, Info));
+				return;
+			}
+
+			TSharedPtr<FJsonValue> OutJson;
+			TSharedRef<TJsonReader<>> Reader = TJsonReaderFactory<>::Create(ResponseStr);
+
+			if (!FJsonSerializer::Deserialize(Reader, OutJson, FJsonSerializer::EFlags::None))
+			{
+				UE_LOG(LogUbaHorde, Error, TEXT("Invalid response body: %s"), *ResponseStr);
+				Promise->SetValue(MakeTuple(HttpResponse, Info));
+				return;
+			}
+
+			TSharedPtr<FJsonValue> NonceValue = OutJson->AsObject()->TryGetField(TEXT("nonce"));
+			TSharedPtr<FJsonValue> IpValue = OutJson->AsObject()->TryGetField(TEXT("ip"));
+			TSharedPtr<FJsonValue> PortValue = OutJson->AsObject()->TryGetField(TEXT("port"));
+
+			if (!NonceValue.Get() || !IpValue.Get() || !PortValue.Get())
+			{
+				UE_LOG(LogUbaHorde, Error, TEXT("Invalid response body: %s"), *ResponseStr);
+				Promise->SetValue(MakeTuple(HttpResponse, Info));
+				return;
+			}
+
+			// Search for property "OSFamily=Windows". If not found, assume we are running on a POSIX system and need Wine to run UBA
+			if (TSharedPtr<FJsonValue> PropertiesValue = OutJson->AsObject()->TryGetField(TEXT("properties")))
+			{
+				for (const TSharedPtr<FJsonValue>& PropertyEntryValue : PropertiesValue->AsArray())
+				{
+					checkf(PropertyEntryValue.Get(), TEXT("null pointer in JSON array object of node \"properties\""));
+					const FString PropertyElementString = PropertyEntryValue->AsString();
+					if (PropertyElementString == TEXT("OSFamily=Windows"))
+					{
+						Info.bRunsWindowOS = true;
+					}
+					if (PropertyElementString.StartsWith(TEXT("LogicalCores=")))
+					{
+						Info.LogicalCores = (uint16)FCString::Atoi(*PropertyElementString + 13);
+					}
+				}
+			}
+
+			FString NonceString = NonceValue->AsString();
+			FString IpString = IpValue->AsString();
+			uint16 PortNumber = (uint16)PortValue->AsNumber();
+
+			if (TSharedPtr<FJsonValue> LeaseIdValue = OutJson->AsObject()->TryGetField(TEXT("leaseId")))
+			{
+				const FString AgentWebPortalUrl = FString::Format(TEXT("{0}/lease/{1}"), { this->ServerUrl, LeaseIdValue->AsString() });
+				UE_LOG(
+					LogUbaHorde, Display, TEXT("UBA Horde machine assigned (%s) [%s:%u]: %s"),
+					(Info.bRunsWindowOS ? TEXT("Microsoft Windows") : TEXT("POSIX/Wine")), *IpString, (uint32)PortNumber, *AgentWebPortalUrl
+				);
+			}
+			else
+			{
+				UE_LOG(LogUbaHorde, Display, TEXT("UBA Horde machine assigned [%s:%u]"), *IpString, (uint32)PortNumber);
+			}
+
+			Info.Ip = IpString;
+			Info.Port = PortNumber;
+			FString::ToHexBlob(NonceString, Info.Nonce, HORDE_NONCE_SIZE);
+
+			Promise->SetValue(MakeTuple(HttpResponse, Info));
+		});
+
+	Request->ProcessRequest();
+
+	return Promise;
+}
diff --git a/Engine/Source/Developer/UbaCoordinatorHorde/Private/UbaHordeMetaClient.h b/Engine/Source/Developer/UbaCoordinatorHorde/Private/UbaHordeMetaClient.h
new file mode 100644
--- /dev/null
+++ b/Engine/Source/Developer/UbaCoordinatorHorde/Private/UbaHordeMetaClient.h
@@ -0,0 +1,63 @@
+// Copyright Epic Games, Inc. All Rights Reserved.
+#pragma once
+
+#include "HttpModule.h"
+#include "Async/Future.h"
+#include "Containers/UnrealString.h"
+#include "Interfaces/IHttpRequest.h"
+#include "Interfaces/IHttpResponse.h"
+#include "HordeHttpClient.h"
+
+DECLARE_LOG_CATEGORY_EXTERN(LogUbaHorde, Log, Log);
+
+#define HORDE_NONCE_SIZE 64
+
+// When asking the Horde server for a machine, this will be returned when the HTTP
+// reponse comes back (after a machine request). If the response didn't come through
+// (i.e., Horde was unable to assign a machine or something), then:
+// - Ip will be set to ""
+// - Port will be set to 0xFFFF
+// - Nonce will be all 0
+struct FHordeRemoteMachineInfo
+{
+	FString Ip;
+	uint16 Port;
+	uint16 LogicalCores;
+	uint8 Nonce[HORDE_NONCE_SIZE];
+	bool bRunsWindowOS;
+};
+
+// This encapsulates the mechanism of talking to the Horde "meta server". The server which
+// can grant us access to remote machines.
+class FUbaHordeMetaClient
+{
+public:
+	// We return the HttpResponse in case more information could be used out of it later.
+	using HordeMachinePromise = TPromise<TTuple<FHttpResponsePtr, FHordeRemoteMachineInfo>>;
+	
+	// Construct the client with a single URL that includes the protocol, IP address, and port in the common format "PROTOCOL://IP:PORT".
+	FUbaHordeMetaClient(const FStringView& HordeServerUrl, const FStringView& InOAuthProviderIdentifier);
+
+	// Local Horde server uses port 5000 by default.
+	FUbaHordeMetaClient(const FStringView& HordeServerIp, uint16 HordeServerPort, bool bInConnectWithAuthentication, const FStringView& OAuthProviderIdentifier);
+
+	bool RefreshHttpClient();
+
+	// This will make a request to Horde for a remote machine to do work on.
+	// Example of actually getting the FHordeRemoteMachineInfo struct:
+	// ```
+	// auto Promise = HordeServer.RequestMachine();
+	// auto Future = Promise->GetFuture();
+	// ... // Can to asynchronous work
+	// auto Future.Wait(); // Wait for the response to arrive
+	// FHordeRemoteMachineInfo MachineInfo = Future.Get().Value;
+	TSharedPtr<HordeMachinePromise, ESPMode::ThreadSafe> RequestMachine(const FString& PoolId, const FString& Machine = "default");
+
+private:
+	void ParseConfig(const FString& IniFilename, const FString& IniSection);
+
+	const FString ServerUrl;
+	const bool bConnectWithAuthentication;
+	TUniquePtr<FHordeHttpClient> HttpClient;
+	FString OAuthProviderIdentifier;
+};
diff --git a/Engine/Source/Developer/UbaCoordinatorHorde/Public/UbaHordeAgentManager.h b/Engine/Source/Developer/UbaCoordinatorHorde/Public/UbaHordeAgentManager.h
new file mode 100644
--- /dev/null
+++ b/Engine/Source/Developer/UbaCoordinatorHorde/Public/UbaHordeAgentManager.h
@@ -0,0 +1,71 @@
+// Copyright Epic Games, Inc. All Rights Reserved.
+#pragma once
+
+#include "Containers/Array.h"
+#include "HAL/Thread.h"
+#include "UbaBase.h"
+
+class FEvent;
+class FUbaHordeMetaClient;
+namespace uba { class NetworkServer; }
+
+class FUbaHordeAgentManager
+{
+public:
+	UBACOORDINATORHORDE_API FUbaHordeAgentManager(const FString& InWorkingDir, const FString& BinariesPath);
+	UBACOORDINATORHORDE_API ~FUbaHordeAgentManager();
+
+	inline void SetUrl(const FString& InUrl) { Url = InUrl; }
+	inline void SetPool(const FString& InPool) { Pool = InPool; }
+	inline void SetOidc(const FString& InOidc) { Oidc = InOidc; }
+	inline void SetMaxCoreCount(uint32 Count) { MaxCores = Count; }
+
+	UBACOORDINATORHORDE_API void SetTargetCoreCount(uint32 Count);
+
+	using AddClientCallback = bool(void* userData, const uba::tchar* ip, uint16 port);
+	UBACOORDINATORHORDE_API void SetAddClientCallback(AddClientCallback* callback, void* userData);
+
+	// Returns the number of agents currently handled by this agent manager.
+	UBACOORDINATORHORDE_API int32 GetAgentCount() const;
+
+	// Returns the active number of cores allocated across all agents.
+	UBACOORDINATORHORDE_API uint32 GetActiveCoreCount() const;
+
+private:
+	struct FHordeAgentWrapper
+	{
+		FThread Thread;
+		FEvent* ShouldExit;
+	};
+
+	void RequestAgent();
+	void ThreadAgent(FHordeAgentWrapper& Wrapper);
+
+	FString WorkingDir;
+	FString BinariesPath;
+
+	FString Url;
+	FString Pool;
+	FString Oidc;
+	uint32 MaxCores = 0;
+
+	TUniquePtr<FUbaHordeMetaClient> HordeMetaClient;
+
+	FCriticalSection BundleRefPathsLock;
+	TArray<FString> BundleRefPaths;
+
+	mutable FCriticalSection AgentsLock;
+	TArray<TUniquePtr<FHordeAgentWrapper>> Agents;
+
+	TAtomic<uint64> LastRequestFailTime;
+	TAtomic<uint32> TargetCoreCount;
+	TAtomic<uint32> EstimatedCoreCount;
+	TAtomic<uint32> ActiveCoreCount;
+	TAtomic<bool> AskForAgents;
+
+	AddClientCallback* m_callback = nullptr;
+	void* m_userData;
+
+	FUbaHordeAgentManager(const FUbaHordeAgentManager&) = delete;
+	void operator=(const FUbaHordeAgentManager&) = delete;
+};
diff --git a/Engine/Source/Developer/UbaCoordinatorHorde/UbaCoordinatorHorde.Build.cs b/Engine/Source/Developer/UbaCoordinatorHorde/UbaCoordinatorHorde.Build.cs
new file mode 100644
--- /dev/null
+++ b/Engine/Source/Developer/UbaCoordinatorHorde/UbaCoordinatorHorde.Build.cs
@@ -0,0 +1,29 @@
+// Copyright Epic Games, Inc. All Rights Reserved.
+
+using System.IO;
+using UnrealBuildTool;
+
+[SupportedPlatforms(UnrealPlatformClass.Desktop)]
+public class UbaCoordinatorHorde : ModuleRules
+{
+	public UbaCoordinatorHorde(ReadOnlyTargetRules Target) : base(Target)
+	{
+		UnsafeTypeCastWarningLevel = WarningLevel.Error;
+		StaticAnalyzerDisabledCheckers.Clear();
+
+		PrivateDependencyModuleNames.AddRange(new string[]
+		{
+			"ApplicationCore",
+			"Core",
+			"DesktopPlatform",
+			"Horde",
+			"HTTP",
+			"Json",
+			"Sockets",
+		});
+
+		string UbaSourcePath = Path.Combine(EngineDirectory, "Source", "Programs", "UnrealBuildAccelerator");
+		PublicIncludePaths.Add(Path.Combine(UbaSourcePath, "Core", "Public"));
+		PublicIncludePaths.Add(Path.Combine(UbaSourcePath, "Common", "Public"));
+	}
+}
diff --git a/Engine/Source/Programs/UnrealBuildAccelerator/Deploy-Linux.bat b/Engine/Source/Programs/UnrealBuildAccelerator/Deploy-Linux.bat
--- a/Engine/Source/Programs/UnrealBuildAccelerator/Deploy-Linux.bat
+++ b/Engine/Source/Programs/UnrealBuildAccelerator/Deploy-Linux.bat
@@ -24,7 +24,8 @@ call Engine/Build/BatchFiles/RunUBT.bat -NoUba -NoUbaLocal -NoSNDBS -NoXGE ^
 	-Target="UbaDetours Linux %config%" ^
 	-Target="UbaHost Linux %config%" ^
 	-Target="UbaTest Linux %config% %options%" ^
-	-Target="UbaTestApp Linux %config% %options%" ^
+	-Target="UbaTestApp Linux %config%" ^
+	-Target="UbaCoordinatorHorde Linux %config% %options%"
 
 if %ERRORLEVEL% NEQ 0 exit /b %ERRORLEVEL%
 popd
diff --git a/Engine/Source/Programs/UnrealBuildAccelerator/UbaCoordinatorHorde.Target.cs b/Engine/Source/Programs/UnrealBuildAccelerator/UbaCoordinatorHorde.Target.cs
new file mode 100644
--- /dev/null
+++ b/Engine/Source/Programs/UnrealBuildAccelerator/UbaCoordinatorHorde.Target.cs
@@ -0,0 +1,72 @@
+// Copyright Epic Games, Inc. All Rights Reserved.
+
+using UnrealBuildTool;
+using System.IO;
+
+[SupportedPlatforms(UnrealPlatformClass.Desktop)]
+public class UbaCoordinatorHordeTarget : TargetRules
+{
+	public UbaCoordinatorHordeTarget(TargetInfo Target) : base(Target)
+	{
+		LaunchModuleName = "UbaCoordinatorHorde";
+		bShouldCompileAsDLL = true;
+
+		Type = TargetType.Program;
+		IncludeOrderVersion = EngineIncludeOrderVersion.Latest;
+		LinkType = TargetLinkType.Monolithic;
+		bDeterministic = true;
+		bWarningsAsErrors = true;
+
+		UndecoratedConfiguration = UnrealTargetConfiguration.Shipping;
+		bHasExports = false;
+
+		// Lean and mean
+		bBuildDeveloperTools = false;
+
+		// Editor-only is enabled for desktop platforms to run unit tests that depend on editor-only data
+		// It's disabled in test and shipping configs to make profiling similar to the game
+		bBuildWithEditorOnlyData = false;
+
+		// Currently this app is not linking against the engine, so we'll compile out references from Core to the rest of the engine
+		bCompileAgainstEngine = false;
+		bCompileAgainstCoreUObject = false;
+		//bCompileAgainstApplicationCore = false;
+		bCompileICU = false;
+
+		bIsBuildingConsoleApplication = true;
+
+		WindowsPlatform.TargetWindowsVersion = 0x0A00;
+		//bUseStaticCRT = true;
+
+		SolutionDirectory = "Programs/UnrealBuildAccelerator";
+
+		var BinariesFolder = Path.Combine("Binaries", Target.Platform.ToString(), "UnrealBuildAccelerator");
+
+		if (Target.Platform.IsInGroup(UnrealPlatformGroup.Windows))
+		{
+			OutputFile = Path.Combine(BinariesFolder, Target.Architectures.SingleArchitecture.bIsX64 ? "x64" : "arm64", $"{LaunchModuleName}{".dll"}");
+		}
+		else
+		{
+			GlobalDefinitions.Add("PLATFORM_WINDOWS=0");
+		}
+
+		if (Target.Platform.IsInGroup(UnrealPlatformGroup.Linux))
+		{
+			OutputFile = Path.Combine(BinariesFolder, $"lib{LaunchModuleName}.so");
+		}
+		else
+		{
+			GlobalDefinitions.Add("PLATFORM_LINUX=0");
+		}
+
+		GlobalDefinitions.Add("UE_EXTERNAL_PROFILING_ENABLED=0");
+		GlobalDefinitions.Add("DISABLE_LEGACY_CORE_TEXTS=1");
+		GlobalDefinitions.Add("STATS=0");
+		GlobalDefinitions.Add("ENABLE_STATNAMEDEVENTS=0");
+		GlobalDefinitions.Add("ALLOW_HITCH_DETECTION=0");
+		GlobalDefinitions.Add("UE_USE_MALLOC_FILL_BYTES=0");
+
+		GlobalDefinitions.Add("UBA_COORDINATOR_HORDE_DLL");
+	}
+}
