[UbaCoordinatorHorde]
* Fixed crash happening when helper was disconnected just when it was connected
* Fixed uninitialized global variable used to by messagebox code

diff --git a/Engine/Source/Developer/UbaCoordinatorHorde/Private/UbaCoordinatorHorde.cpp b/Engine/Source/Developer/UbaCoordinatorHorde/Private/UbaCoordinatorHorde.cpp
--- a/Engine/Source/Developer/UbaCoordinatorHorde/Private/UbaCoordinatorHorde.cpp
+++ b/Engine/Source/Developer/UbaCoordinatorHorde/Private/UbaCoordinatorHorde.cpp
@@ -54,6 +54,8 @@ extern "C"
 	UBA_EXPORT uba::Coordinator* UbaCreateCoordinator(const uba::CoordinatorCreateInfo& info)
 	{
 		FCommandLine::Set(TEXT(""));
+		GWarn = FPlatformApplicationMisc::GetFeedbackContext();
+
 		if (info.logging)
 		{
 			FConfigCacheIni::InitializeConfigSystem();
diff --git a/Engine/Source/Developer/UbaCoordinatorHorde/Private/UbaHordeAgentManager.cpp b/Engine/Source/Developer/UbaCoordinatorHorde/Private/UbaHordeAgentManager.cpp
--- a/Engine/Source/Developer/UbaCoordinatorHorde/Private/UbaHordeAgentManager.cpp
+++ b/Engine/Source/Developer/UbaCoordinatorHorde/Private/UbaHordeAgentManager.cpp
@@ -298,6 +298,11 @@ void FUbaHordeAgentManager::ThreadAgent(FHordeAgentWrapper& Wrapper)
 
 		Agent = MakeUnique<FUbaHordeAgent>(MachineInfo);
 
+		if (!Agent->IsValid())
+		{
+			return;
+		}
+
 		if (!Agent->BeginCommunication())
 		{
 			return;
