Fix deadstore analysis warning

#rnx

diff --git a/Engine/Source/Developer/UbaCoordinatorHorde/Private/UbaHordeComputeTransport.cpp b/Engine/Source/Developer/UbaCoordinatorHorde/Private/UbaHordeComputeTransport.cpp
--- a/Engine/Source/Developer/UbaCoordinatorHorde/Private/UbaHordeComputeTransport.cpp
+++ b/Engine/Source/Developer/UbaCoordinatorHorde/Private/UbaHordeComputeTransport.cpp
@@ -42,7 +42,6 @@ size_t FUbaHordeComputeTransport::Send(const void* Data, size_t Size)
 	int32 NumBytesSent = 0;
 	if (!Socket->Send((const uint8*)Data, Size, NumBytesSent))
 	{
-		NumBytesSent = 0;
 		// We can't log this since the other side could have disconnected us (causing recv to fail). This happens often in the horde setup
 		//UE_LOG(LogUbaHorde, Error, TEXT("Failed to send data to the Horde Agent: %llu %s"), (uint64)Size, Size == 1 ? TEXT("byte") : TEXT("bytes"));
 		bHasErrors = true;
@@ -60,8 +59,6 @@ size_t FUbaHordeComputeTransport::Recv(void* Data, size_t Size)
 	int32 NumBytesRead = 0;
 	if (!Socket->Recv((uint8*)Data, Size, NumBytesRead))
 	{
-		NumBytesRead = 0;
-
 		// We can't log this since the other side could have disconnected us (causing recv to fail). This happens often in the horde setup
 		//if (!bIsClosed)
 		//	UE_LOG(LogUbaHorde, Error, TEXT("Failed to receive data from the Horde Agent: %llu %s"), (uint64)Size, Size == 1 ? TEXT("byte") : TEXT("bytes"));
