[UBA] Use new Horde API to retrieve server URL and replace OIDC token login.

This unifies the Horde agent connection in the UbaController and avoids a custom configuration for the Horde server.

#rnx

diff --git a/Engine/Plugins/UbaController/Source/Private/UbaJobProcessor.cpp b/Engine/Plugins/UbaController/Source/Private/UbaJobProcessor.cpp
--- a/Engine/Plugins/UbaController/Source/Private/UbaJobProcessor.cpp
+++ b/Engine/Plugins/UbaController/Source/Private/UbaJobProcessor.cpp
@@ -354,13 +354,6 @@ void FUbaJobProcessor::StartUba()
 		HordeConfig.RemoveFromStart(TEXT("("));
 		HordeConfig.RemoveFromEnd(TEXT(")"));
 
-		FString Url;
-		if (FParse::Value(*HordeConfig, TEXT("Url="), Url))
-		{
-			UE_LOG(LogUbaController, Log, TEXT("Found UBA controller Url: \"%s\""), *Url);
-			HordeAgentManager->SetUrl(Url);
-		}
-
 		FString Pool;
 		if (FParse::Value(*HordeConfig, TEXT("Pool="), Pool))
 		{
@@ -368,13 +361,6 @@ void FUbaJobProcessor::StartUba()
 			HordeAgentManager->SetPool(Pool);
 		}
 
-		FString Oidc;
-		if (FParse::Value(*HordeConfig, TEXT("Oidc="), Oidc))
-		{
-			UE_LOG(LogUbaController, Log, TEXT("Found UBA controller Oidc: \"%s\""), *Oidc);
-			HordeAgentManager->SetOidc(Oidc);
-		}
-
 		uint32 MaxCores = 0;
 		if (FParse::Value(*HordeConfig, TEXT("MaxCores="), MaxCores))
 		{
diff --git a/Engine/Source/Developer/UbaCoordinatorHorde/Private/UbaCoordinatorHorde.cpp b/Engine/Source/Developer/UbaCoordinatorHorde/Private/UbaCoordinatorHorde.cpp
--- a/Engine/Source/Developer/UbaCoordinatorHorde/Private/UbaCoordinatorHorde.cpp
+++ b/Engine/Source/Developer/UbaCoordinatorHorde/Private/UbaCoordinatorHorde.cpp
@@ -69,9 +69,7 @@ extern "C"
 
 		auto coordinator = new CoordinatorHorde(info.workDir, info.binariesDir);
 		auto& m = coordinator->m_manager;
-		m.SetUrl(info.uri);
 		m.SetPool(info.pool);
-		m.SetOidc(info.oidc);
 		m.SetMaxCoreCount(info.maxCoreCount);
 
 		return coordinator;
diff --git a/Engine/Source/Developer/UbaCoordinatorHorde/Private/UbaHordeAgentManager.cpp b/Engine/Source/Developer/UbaCoordinatorHorde/Private/UbaHordeAgentManager.cpp
--- a/Engine/Source/Developer/UbaCoordinatorHorde/Private/UbaHordeAgentManager.cpp
+++ b/Engine/Source/Developer/UbaCoordinatorHorde/Private/UbaHordeAgentManager.cpp
@@ -36,7 +36,7 @@ FUbaHordeAgentManager::FUbaHordeAgentManager(const FString& InWorkingDir, const
 	,	TargetCoreCount(0)
 	,	EstimatedCoreCount(0)
 	,   ActiveCoreCount(0)
-	,	AskForAgents(true)
+	,	bAskForAgents(true)
 {
 }
 
@@ -55,14 +55,9 @@ void FUbaHordeAgentManager::SetTargetCoreCount(uint32 Count)
 {
 	TargetCoreCount = FMath::Min(MaxCores, Count);
 
-	if (Url.IsEmpty())
-	{
-		AskForAgents = false;
-	}
-
 	while (EstimatedCoreCount < TargetCoreCount)
 	{
-		if (!AskForAgents)
+		if (!bAskForAgents)
 		{
 			return;
 		}
@@ -223,7 +218,7 @@ void FUbaHordeAgentManager::ThreadAgent(FHordeAgentWrapper& Wrapper)
 				if (!CreateHordeBundleFromFile(*FilePath, *BundlePath))
 				{
 					UE_LOG(LogUbaHorde, Error, TEXT("Failed to create Horde bundle for: %s"), *FilePath);
-					AskForAgents = false;
+					bAskForAgents = false;
 					return;
 				}
 				UE_LOG(LogUbaHorde, Display, TEXT("Created Horde bundle for: %s"), *FilePath);
@@ -234,16 +229,16 @@ void FUbaHordeAgentManager::ThreadAgent(FHordeAgentWrapper& Wrapper)
 		if (!HordeMetaClient)
 		{
 			// Create Horde meta client right before we need it to make sure the CVar for the server URL has been read by now
-			HordeMetaClient = MakeUnique<FUbaHordeMetaClient>(Url, Oidc);
+			HordeMetaClient = MakeUnique<FUbaHordeMetaClient>();
 			if (!HordeMetaClient->RefreshHttpClient())
 			{
 				UE_LOG(LogUbaHorde, Error, TEXT("Failed to create HttpClient for UbaAgent"));
-				AskForAgents = false;
+				bAskForAgents = false;
 				return;
 			}
 		}
 
-		if (!AskForAgents)
+		if (!bAskForAgents)
 		{
 			return;
 		}
diff --git a/Engine/Source/Developer/UbaCoordinatorHorde/Private/UbaHordeMetaClient.cpp b/Engine/Source/Developer/UbaCoordinatorHorde/Private/UbaHordeMetaClient.cpp
--- a/Engine/Source/Developer/UbaCoordinatorHorde/Private/UbaHordeMetaClient.cpp
+++ b/Engine/Source/Developer/UbaCoordinatorHorde/Private/UbaHordeMetaClient.cpp
@@ -1,6 +1,7 @@
 // Copyright Epic Games, Inc. All Rights Reserved.
 
 #include "UbaHordeMetaClient.h"
+#include "Horde.h"
 #include "Serialization/JsonReader.h"
 #include "Serialization/JsonSerializer.h"
 #include "Misc/App.h"
@@ -8,52 +9,23 @@
 
 DEFINE_LOG_CATEGORY(LogUbaHorde);
 
-FUbaHordeMetaClient::FUbaHordeMetaClient(const FStringView& HordeServerUrl, const FStringView& InOAuthProviderIdentifier)
-	: ServerUrl(HordeServerUrl)
-	, bConnectWithAuthentication(HordeServerUrl.StartsWith(TEXT("https:")))
-	, OAuthProviderIdentifier(InOAuthProviderIdentifier)
-{
-}
-
-FUbaHordeMetaClient::FUbaHordeMetaClient(const FStringView& HordeServerIp, uint16 HordeServerPort, bool bInConnectWithAuthentication, const FStringView& InOAuthProviderIdentifier)
-	: ServerUrl(FString::Format(TEXT("{0}://{1}:{2}"), {(bInConnectWithAuthentication ? TEXT("https") : TEXT("http")), HordeServerIp, HordeServerPort}))
-	, bConnectWithAuthentication(bInConnectWithAuthentication)
-	, OAuthProviderIdentifier(InOAuthProviderIdentifier)
-{
-}
-
 bool FUbaHordeMetaClient::RefreshHttpClient()
 {
+	if (!FHorde::GetServerUrl(ServerUrl))
+	{
+		UE_LOG(LogUbaHorde, Warning, TEXT("Getting Horde server URL failed"));
+		return false;
+	}
+
 	// Try to connect to Horde with HTTP and v2 API
 	HttpClient = MakeUnique<FHordeHttpClient>(ServerUrl);
 
-	if (bConnectWithAuthentication)
+	if (!HttpClient->Login(FApp::IsUnattended()))
 	{
-		check(!OAuthProviderIdentifier.IsEmpty());
-		if (FApp::IsUnattended())
-		{
-			UE_LOG(LogUbaHorde, Display, TEXT("Logging in to Horde server with environment variable UE_HORDE_TOKEN: %s"), *ServerUrl);
-			if (HttpClient->LoginWithEnvironmentVariable())
-			{
-				return true;
-			}
-
-			// Fallback to OIDC login as some commandlets might add -Unattended implicitly, but not everyone has UE_HORDE_TOKEN environment variable available
-			UE_LOG(LogUbaHorde, Display, TEXT("Login attempt with environment variable failed; trying to login with OIDC: %s"), *ServerUrl);
-		}
-		else
-		{
-			UE_LOG(LogUbaHorde, Display, TEXT("Logging in to Horde server with OIDC: %s"), *ServerUrl);
-		}
-
-		if (HttpClient->LoginWithOidc(*OAuthProviderIdentifier, FApp::IsUnattended()))
-		{
-			return true;
-		}
-
 		UE_LOG(LogUbaHorde, Warning, TEXT("Login to Horde server [%s] failed"), *ServerUrl);
 		return false;
 	}
+
 	return true;
 }
 
diff --git a/Engine/Source/Developer/UbaCoordinatorHorde/Private/UbaHordeMetaClient.h b/Engine/Source/Developer/UbaCoordinatorHorde/Private/UbaHordeMetaClient.h
--- a/Engine/Source/Developer/UbaCoordinatorHorde/Private/UbaHordeMetaClient.h
+++ b/Engine/Source/Developer/UbaCoordinatorHorde/Private/UbaHordeMetaClient.h
@@ -29,17 +29,13 @@ struct FHordeRemoteMachineInfo
 
 // This encapsulates the mechanism of talking to the Horde "meta server". The server which
 // can grant us access to remote machines.
-class FUbaHordeMetaClient
+class FUbaHordeMetaClient final
 {
 public:
 	// We return the HttpResponse in case more information could be used out of it later.
 	using HordeMachinePromise = TPromise<TTuple<FHttpResponsePtr, FHordeRemoteMachineInfo>>;
 	
-	// Construct the client with a single URL that includes the protocol, IP address, and port in the common format "PROTOCOL://IP:PORT".
-	FUbaHordeMetaClient(const FStringView& HordeServerUrl, const FStringView& InOAuthProviderIdentifier);
-
-	// Local Horde server uses port 5000 by default.
-	FUbaHordeMetaClient(const FStringView& HordeServerIp, uint16 HordeServerPort, bool bInConnectWithAuthentication, const FStringView& OAuthProviderIdentifier);
+	FUbaHordeMetaClient() = default;
 
 	bool RefreshHttpClient();
 
@@ -56,8 +52,6 @@ public:
 private:
 	void ParseConfig(const FString& IniFilename, const FString& IniSection);
 
-	const FString ServerUrl;
-	const bool bConnectWithAuthentication;
+	FString ServerUrl;
 	TUniquePtr<FHordeHttpClient> HttpClient;
-	FString OAuthProviderIdentifier;
 };
diff --git a/Engine/Source/Developer/UbaCoordinatorHorde/Public/UbaHordeAgentManager.h b/Engine/Source/Developer/UbaCoordinatorHorde/Public/UbaHordeAgentManager.h
--- a/Engine/Source/Developer/UbaCoordinatorHorde/Public/UbaHordeAgentManager.h
+++ b/Engine/Source/Developer/UbaCoordinatorHorde/Public/UbaHordeAgentManager.h
@@ -15,9 +15,7 @@ public:
 	UBACOORDINATORHORDE_API FUbaHordeAgentManager(const FString& InWorkingDir, const FString& BinariesPath);
 	UBACOORDINATORHORDE_API ~FUbaHordeAgentManager();
 
-	inline void SetUrl(const FString& InUrl) { Url = InUrl; }
 	inline void SetPool(const FString& InPool) { Pool = InPool; }
-	inline void SetOidc(const FString& InOidc) { Oidc = InOidc; }
 	inline void SetMaxCoreCount(uint32 Count) { MaxCores = Count; }
 
 	UBACOORDINATORHORDE_API void SetTargetCoreCount(uint32 Count);
@@ -44,9 +42,7 @@ private:
 	FString WorkingDir;
 	FString BinariesPath;
 
-	FString Url;
 	FString Pool;
-	FString Oidc;
 	uint32 MaxCores = 0;
 
 	TUniquePtr<FUbaHordeMetaClient> HordeMetaClient;
@@ -61,7 +57,7 @@ private:
 	TAtomic<uint32> TargetCoreCount;
 	TAtomic<uint32> EstimatedCoreCount;
 	TAtomic<uint32> ActiveCoreCount;
-	TAtomic<bool> AskForAgents;
+	TAtomic<bool> bAskForAgents;
 
 	AddClientCallback* m_callback = nullptr;
 	void* m_userData;
diff --git a/Engine/Source/Programs/UnrealBuildAccelerator/Cli/Private/UbaCli.cpp b/Engine/Source/Programs/UnrealBuildAccelerator/Cli/Private/UbaCli.cpp
--- a/Engine/Source/Programs/UnrealBuildAccelerator/Cli/Private/UbaCli.cpp
+++ b/Engine/Source/Programs/UnrealBuildAccelerator/Cli/Private/UbaCli.cpp
@@ -150,9 +150,7 @@ namespace uba
 		StringBuffer<128> listenIp;
 		StringBuffer<128> cacheHost;
 		TString coordinatorName;
-		TString coordinatorUri;
 		TString coordinatorPool;
-		TString coordinatorOidc;
 		u32 coordinatorMaxCoreCount = 400;
 		u16 port = DefaultPort;
 		u32 maxProcessCount = DefaultProcessorCount;
@@ -244,24 +242,12 @@ namespace uba
 					return PrintHelp(TC("-coordinator needs a value"));
 				coordinatorName = value.data;
 			}
-			else if (name.Equals(TC("-uri")))
-			{
-				if (value.IsEmpty())
-					return PrintHelp(TC("-uri needs a value"));
-				coordinatorUri = value.data;
-			}
 			else if (name.Equals(TC("-pool")))
 			{
 				if (value.IsEmpty())
 					return PrintHelp(TC("-pool needs a value"));
 				coordinatorPool = value.data;
 			}
-			else if (name.Equals(TC("-oidc")))
-			{
-				if (value.IsEmpty())
-					return PrintHelp(TC("-oidc needs a value"));
-				coordinatorOidc = value.data;
-			}
 			else if (name.Equals(TC("-maxcores")))
 			{
 				if (value.IsEmpty())
@@ -771,9 +757,7 @@ namespace uba
 			cinfo.binariesDir = binariesDir.data;
 
 			// TODO: This is very horde specific.. maybe all these parameters should be a string or something
-			cinfo.uri = coordinatorUri.c_str();
 			cinfo.pool = coordinatorPool.c_str();
-			cinfo.oidc = coordinatorOidc.c_str();
 			cinfo.maxCoreCount = 400;
 			cinfo.logging = true;
 			if (!coordinator.Create(logger, coordinatorName.c_str(), cinfo, networkBackend, networkServer))
diff --git a/Engine/Source/Programs/UnrealBuildAccelerator/Common/Public/UbaCoordinator.h b/Engine/Source/Programs/UnrealBuildAccelerator/Common/Public/UbaCoordinator.h
--- a/Engine/Source/Programs/UnrealBuildAccelerator/Common/Public/UbaCoordinator.h
+++ b/Engine/Source/Programs/UnrealBuildAccelerator/Common/Public/UbaCoordinator.h
@@ -10,9 +10,7 @@ namespace uba
 		const tchar* binariesDir = nullptr;
 
 		// TODO: This is very horde specific.. maybe all these parameters should be a string or something
-		const tchar* uri = nullptr;
 		const tchar* pool = nullptr;
-		const tchar* oidc = nullptr;
 		u32 maxCoreCount = 0;
 
 		bool logging = false;
diff --git a/Engine/Source/Programs/UnrealBuildAccelerator/Test/Private/UbaTestStress.h b/Engine/Source/Programs/UnrealBuildAccelerator/Test/Private/UbaTestStress.h
--- a/Engine/Source/Programs/UnrealBuildAccelerator/Test/Private/UbaTestStress.h
+++ b/Engine/Source/Programs/UnrealBuildAccelerator/Test/Private/UbaTestStress.h
@@ -72,9 +72,7 @@ namespace uba
 			cinfo.workDir = coordinatorWorkDir.data;
 			cinfo.binariesDir = binariesDir.data;
 
-			cinfo.uri = TC("https://horde-uri");
 			cinfo.pool = TC("PoolToUse");
-			cinfo.oidc = TC("Oidc");
 
 
 			cinfo.maxCoreCount = 400;
