[UBA Controller] Adding support for running uba agent with -Host

Also added support to specify a uba port and force remote via the existing UBA Controller section in Engine.ini

diff --git a/Engine/Plugins/UbaController/Source/Private/UbaJobProcessor.cpp b/Engine/Plugins/UbaController/Source/Private/UbaJobProcessor.cpp
--- a/Engine/Plugins/UbaController/Source/Private/UbaJobProcessor.cpp
+++ b/Engine/Plugins/UbaController/Source/Private/UbaJobProcessor.cpp
@@ -369,6 +369,24 @@ void FUbaJobProcessor::StartUba()
 		}
 	}
 
+	FString Host;
+	if (GConfig->GetString(TEXT("UbaController"), TEXT("Host"), Host, GEngineIni))
+	{
+		UE_LOG(LogUbaController, Log, TEXT("Found UBA controller Host: \"%s\""), *Host);
+		HordeAgentManager->SetUbaHost(StringCast<ANSICHAR>(*Host).Get());
+	}
+
+	int32 Port = 0;
+	if (GConfig->GetInt(TEXT("UbaController"), TEXT("Port"), Port, GEngineIni))
+	{
+		UE_LOG(LogUbaController, Log, TEXT("Found UBA controller Port: \"%d\""), Port);
+		HordeAgentManager->SetUbaPort(Port);
+	}
+	
+	if (GConfig->GetBool(TEXT("UbaController"), TEXT("bForceRemote"), bForceRemote, GEngineIni))
+	{
+		UE_LOG(LogUbaController, Log, TEXT("Found UBA controller Force Remote: [%s]"), bForceRemote ? TEXT("True") : TEXT("False"));
+	}
 
 	UE_LOG(LogUbaController, Display, TEXT("Created UBA storage server: RootDir=%s"), *RootDir);
 }
@@ -448,7 +466,7 @@ uint32 FUbaJobProcessor::Run()
 			int32 LocalCoresToNotUse = 1 + activeRemote / 30; // Use one core per 30 remote ones
 			int32 MaxLocal = FMath::Max(0, MaxPossible - LocalCoresToNotUse);
 			MaxLocal = FMath::Min(MaxLocal, MaxLocalParallelJobs);
-			Scheduler_SetMaxLocalProcessors(UbaScheduler, MaxLocal);
+			Scheduler_SetMaxLocalProcessors(UbaScheduler, bForceRemote ? 0 : MaxLocal);
 
 			int32 TargetCoreCount = FMath::Max(0, int32(queued + active) - MaxLocal);
 
diff --git a/Engine/Plugins/UbaController/Source/Public/UbaJobProcessor.h b/Engine/Plugins/UbaController/Source/Public/UbaJobProcessor.h
--- a/Engine/Plugins/UbaController/Source/Public/UbaJobProcessor.h
+++ b/Engine/Plugins/UbaController/Source/Public/UbaJobProcessor.h
@@ -83,4 +83,7 @@ private:
 
 	FCriticalSection StatsLock;
 	FDistributedBuildStats Stats;
+
+	/** If true all UBA jobs will be run remotely */
+	bool bForceRemote = false;
 };
\ No newline at end of file
diff --git a/Engine/Source/Developer/UbaCoordinatorHorde/Private/UbaHordeAgentManager.cpp b/Engine/Source/Developer/UbaCoordinatorHorde/Private/UbaHordeAgentManager.cpp
--- a/Engine/Source/Developer/UbaCoordinatorHorde/Private/UbaHordeAgentManager.cpp
+++ b/Engine/Source/Developer/UbaCoordinatorHorde/Private/UbaHordeAgentManager.cpp
@@ -186,7 +186,9 @@ void FUbaHordeAgentManager::ThreadAgent(FHordeAgentWrapper& Wrapper)
 	};
 
 	int MachineCoreCount = 0;
-	uint32 ListenPort = 7001;
+
+	// If no host is specified, we need to start the agent in listen mode
+	const bool bUseListen = UbaHost.IsEmpty();
 
 	{
 		ON_SCOPE_EXIT{ EstimatedCoreCount -= 32; };
@@ -327,11 +329,12 @@ void FUbaHordeAgentManager::ThreadAgent(FHordeAgentWrapper& Wrapper)
 		}
 
 		// Start the UBA Agent that will connect to us, requesting for work
-		const std::string ListenPortArg = "-listen=" + std::to_string(ListenPort);
+
+		const FAnsiString AgentConnectionArg = bUseListen ? FAnsiString::Printf("-listen=%u", UbaPort) : FAnsiString::Printf("-Host=%s:%u", *UbaHost, UbaPort);
 
 		const char* UbaAgentArgs[] =
 		{
-			ListenPortArg.c_str(),
+			*AgentConnectionArg,
 			"-nopoll",				// -nopoll recommended when running on remote Horde agents to make sure they exit after completion. Otherwise, it keeps running.
 			"-listenTimeout=5",		// Agent will wait 5 seconds for this thread to connect (Server_AddClient does the connect)
 			"-quiet",				// Skip all the agent logging that would be sent over to here
@@ -362,7 +365,7 @@ void FUbaHordeAgentManager::ThreadAgent(FHordeAgentWrapper& Wrapper)
 			UbaAgentCmdArgs += TEXT(" ");
 			UbaAgentCmdArgs += ANSI_TO_TCHAR(Arg);
 		}
-		UE_LOG(LogUbaHorde, Log, TEXT("Remote execution on Horde machine [%s:%d]: %s"), *Agent->GetMachineInfo().Ip, ListenPort, *UbaAgentCmdArgs);
+		UE_LOG(LogUbaHorde, Log, TEXT("Remote execution on Horde machine [%s:%u]: %s"), *Agent->GetMachineInfo().Ip, UbaPort, *UbaAgentCmdArgs);
 
 		MachineCoreCount = MachineInfo.LogicalCores;
 		EstimatedCoreCount += MachineCoreCount;
@@ -375,15 +378,20 @@ void FUbaHordeAgentManager::ThreadAgent(FHordeAgentWrapper& Wrapper)
 	{
 		Agent->Poll(UbaCoordinatorHordeModule::bHordeForwardAgentLogs);
 
+		if (!bUseListen)
+		{
+			continue;
+		}
+
 		if (callCounter++ == 2)
 		{
 			// Add this machine as client to the remote agent
 			const FString& IpAddress = Agent->GetMachineInfo().Ip;
-			const bool bAddClientSuccess = 	m_callback(m_userData, StringCast<uba::tchar>(*IpAddress).Get(), (uint16)ListenPort);
+			const bool bAddClientSuccess = 	m_callback(m_userData, StringCast<uba::tchar>(*IpAddress).Get(), static_cast<uint16>(UbaPort));
 
 			if (!bAddClientSuccess)
 			{
-				UE_LOG(LogUbaHorde, Display, TEXT("Server_AddClient(%s:%d) failed"), *IpAddress, ListenPort);
+				UE_LOG(LogUbaHorde, Display, TEXT("Server_AddClient(%s:%u) failed"), *IpAddress, UbaPort);
 				return;
 			}
 		}
diff --git a/Engine/Source/Developer/UbaCoordinatorHorde/Public/UbaHordeAgentManager.h b/Engine/Source/Developer/UbaCoordinatorHorde/Public/UbaHordeAgentManager.h
--- a/Engine/Source/Developer/UbaCoordinatorHorde/Public/UbaHordeAgentManager.h
+++ b/Engine/Source/Developer/UbaCoordinatorHorde/Public/UbaHordeAgentManager.h
@@ -1,6 +1,7 @@
 // Copyright Epic Games, Inc. All Rights Reserved.
 #pragma once
 
+#include "Containers/AnsiString.h"
 #include "Containers/Array.h"
 #include "HAL/Thread.h"
 #include "UbaBase.h"
@@ -17,6 +18,8 @@ public:
 
 	inline void SetPool(const FString& InPool) { Pool = InPool; }
 	inline void SetMaxCoreCount(uint32 Count) { MaxCores = Count; }
+	inline void SetUbaHost(const FAnsiString& InHost) { UbaHost = InHost; }
+	inline void SetUbaPort(uint32 InPort) { UbaPort = InPort; }
 
 	UBACOORDINATORHORDE_API void SetTargetCoreCount(uint32 Count);
 
@@ -43,6 +46,8 @@ private:
 	FString BinariesPath;
 
 	FString Pool;
+	FAnsiString UbaHost;
+	uint32 UbaPort = 7001;
 	uint32 MaxCores = 0;
 
 	TUniquePtr<FUbaHordeMetaClient> HordeMetaClient;
@@ -60,7 +65,7 @@ private:
 	TAtomic<bool> bAskForAgents;
 
 	AddClientCallback* m_callback = nullptr;
-	void* m_userData;
+	void* m_userData = nullptr;
 
 	FUbaHordeAgentManager(const FUbaHordeAgentManager&) = delete;
 	void operator=(const FUbaHordeAgentManager&) = delete;
