[UBA] Lower verbosity for invalid response body of UBA Horde job as this should not fail a CIS job.

This will report an invalid response body (which we see from an HTTP 502 Bad Gateway error) as Display instead of Error.
FUbaHordeMetaClient will handle this event just as if no Horde resources are available.

#rnx

diff --git a/Engine/Source/Developer/UbaCoordinatorHorde/Private/UbaHordeMetaClient.cpp b/Engine/Source/Developer/UbaCoordinatorHorde/Private/UbaHordeMetaClient.cpp
--- a/Engine/Source/Developer/UbaCoordinatorHorde/Private/UbaHordeMetaClient.cpp
+++ b/Engine/Source/Developer/UbaCoordinatorHorde/Private/UbaHordeMetaClient.cpp
@@ -70,7 +70,7 @@ TSharedPtr<FUbaHordeMetaClient::HordeMachinePromise, ESPMode::ThreadSafe> FUbaHo
 
 			FString ResponseStr = HttpResponse->GetContentAsString();
 
-			if (HttpResponse->GetResponseCode() == 503)
+			if (HttpResponse->GetResponseCode() == 503) // HTTP 503 Service Unavailable
 			{
 				UE_LOG(LogUbaHorde, Verbose, TEXT("No resources available in Horde (%s)"), *ResponseStr);
 				Promise->SetValue(MakeTuple(HttpResponse, Info));
@@ -82,7 +82,8 @@ TSharedPtr<FUbaHordeMetaClient::HordeMachinePromise, ESPMode::ThreadSafe> FUbaHo
 
 			if (!FJsonSerializer::Deserialize(Reader, OutJson, FJsonSerializer::EFlags::None))
 			{
-				UE_LOG(LogUbaHorde, Error, TEXT("Invalid response body: %s"), *ResponseStr);
+				// Report invalid response body with Display verbosity only, since this should not fail a CIS job
+				UE_LOG(LogUbaHorde, Display, TEXT("Invalid response body: %s"), *ResponseStr);
 				Promise->SetValue(MakeTuple(HttpResponse, Info));
 				return;
 			}
@@ -93,7 +94,8 @@ TSharedPtr<FUbaHordeMetaClient::HordeMachinePromise, ESPMode::ThreadSafe> FUbaHo
 
 			if (!NonceValue.Get() || !IpValue.Get() || !PortValue.Get())
 			{
-				UE_LOG(LogUbaHorde, Error, TEXT("Invalid response body: %s"), *ResponseStr);
+				// Report invalid response body with Display verbosity only, since this should not fail a CIS job
+				UE_LOG(LogUbaHorde, Display, TEXT("Invalid response body: %s"), *ResponseStr);
 				Promise->SetValue(MakeTuple(HttpResponse, Info));
 				return;
 			}
