[UBA] Fix null pointer exception in UbaHordeAgent.

The two channel objects can be null in CloseConnectioin() when the compute transport object was not valid and the constructor exited early.

#rnx

diff --git a/Engine/Source/Developer/UbaCoordinatorHorde/Private/UbaHordeAgent.cpp b/Engine/Source/Developer/UbaCoordinatorHorde/Private/UbaHordeAgent.cpp
--- a/Engine/Source/Developer/UbaCoordinatorHorde/Private/UbaHordeAgent.cpp
+++ b/Engine/Source/Developer/UbaCoordinatorHorde/Private/UbaHordeAgent.cpp
@@ -171,8 +171,15 @@ void FUbaHordeAgent::Execute(const char* Exe, const char** Args, size_t NumArgs,
 
 void FUbaHordeAgent::CloseConnection()
 {
-	ChildChannel->Close();
-	AgentChannel->Close();
+	// These can be null when the compute transport object was not valid in the constructor
+	if (ChildChannel)
+	{
+		ChildChannel->Close();
+	}
+	if (AgentChannel)
+	{
+		AgentChannel->Close();
+	}
 }
 
 bool FUbaHordeAgent::IsValid()
