[Horde]
* Fixed logging so it is not showing wrong os family

diff --git a/Engine/Source/Developer/UbaCoordinatorHorde/Private/UbaHordeMetaClient.cpp b/Engine/Source/Developer/UbaCoordinatorHorde/Private/UbaHordeMetaClient.cpp
--- a/Engine/Source/Developer/UbaCoordinatorHorde/Private/UbaHordeMetaClient.cpp
+++ b/Engine/Source/Developer/UbaCoordinatorHorde/Private/UbaHordeMetaClient.cpp
@@ -100,16 +100,19 @@ TSharedPtr<FUbaHordeMetaClient::HordeMachinePromise, ESPMode::ThreadSafe> FUbaHo
 				return;
 			}
 
-			// Search for property "OSFamily=Windows". If not found, assume we are running on a POSIX system and need Wine to run UBA
+			FString OsFamily(TEXT("UNKNOWN-OS"));
+
 			if (TSharedPtr<FJsonValue> PropertiesValue = OutJson->AsObject()->TryGetField(TEXT("properties")))
 			{
 				for (const TSharedPtr<FJsonValue>& PropertyEntryValue : PropertiesValue->AsArray())
 				{
 					checkf(PropertyEntryValue.Get(), TEXT("null pointer in JSON array object of node \"properties\""));
 					const FString PropertyElementString = PropertyEntryValue->AsString();
-					if (PropertyElementString == TEXT("OSFamily=Windows"))
+					if (PropertyElementString.StartsWith(TEXT("OSFamily=")))
 					{
-						Info.bRunsWindowOS = true;
+						OsFamily = *PropertyElementString + 9;
+						if (OsFamily == TEXT("Windows"))
+							Info.bRunsWindowOS = true;
 					}
 					if (PropertyElementString.StartsWith(TEXT("LogicalCores=")))
 					{
@@ -126,8 +129,7 @@ TSharedPtr<FUbaHordeMetaClient::HordeMachinePromise, ESPMode::ThreadSafe> FUbaHo
 			{
 				const FString AgentWebPortalUrl = FString::Format(TEXT("{0}lease/{1}"), { this->ServerUrl, LeaseIdValue->AsString() });
 				UE_LOG(
-					LogUbaHorde, Display, TEXT("UBA Horde machine assigned (%s) [%s:%u]: %s"),
-					(Info.bRunsWindowOS ? TEXT("Microsoft Windows") : TEXT("POSIX/Wine")), *IpString, (uint32)PortNumber, *AgentWebPortalUrl
+					LogUbaHorde, Display, TEXT("UBA Horde machine assigned (%s) [%s:%u]: %s"), *OsFamily, *IpString, (uint32)PortNumber, *AgentWebPortalUrl
 				);
 			}
 			else
