[Horde/UBA]
* Fix to prevent crash when horde client fails to setup a proper connection to server

diff --git a/Engine/Source/Developer/Horde/Private/Compute/ComputeSocket.cpp b/Engine/Source/Developer/Horde/Private/Compute/ComputeSocket.cpp
--- a/Engine/Source/Developer/Horde/Private/Compute/ComputeSocket.cpp
+++ b/Engine/Source/Developer/Horde/Private/Compute/ComputeSocket.cpp
@@ -25,13 +25,13 @@ TSharedPtr<FComputeChannel> FComputeSocket::CreateChannel(int ChannelId)
 	FComputeBuffer RecvBuffer;
 	if (!RecvBuffer.CreateNew(FComputeBuffer::FParams()))
 	{
-		return MakeShared<FComputeChannel>();
+		return {};
 	}
 
 	FComputeBuffer SendBuffer;
 	if (!SendBuffer.CreateNew(FComputeBuffer::FParams()))
 	{
-		return MakeShared<FComputeChannel>();
+		return {};
 	}
 
 	return CreateChannel(ChannelId, std::move(RecvBuffer), std::move(SendBuffer));
diff --git a/Engine/Source/Developer/UbaCoordinatorHorde/Private/UbaHordeAgent.cpp b/Engine/Source/Developer/UbaCoordinatorHorde/Private/UbaHordeAgent.cpp
--- a/Engine/Source/Developer/UbaCoordinatorHorde/Private/UbaHordeAgent.cpp
+++ b/Engine/Source/Developer/UbaCoordinatorHorde/Private/UbaHordeAgent.cpp
@@ -32,6 +32,11 @@ FUbaHordeAgent::FUbaHordeAgent(const FHordeRemoteMachineInfo& InMachineInfo)
 	TSharedPtr<FComputeChannel> ComputeChannel = HordeComputeSocket->CreateChannel(0);
 	TSharedPtr<FComputeChannel> SecondComputeChannel = HordeComputeSocket->CreateChannel(100);
 
+	if (!ComputeChannel.IsValid() || !SecondComputeChannel.IsValid())
+	{
+		return;
+	}
+
 	// Now this is what is going to handle the ComputeChannel.
 	AgentChannel = MakeUnique<FAgentMessageChannel>(MoveTemp(ComputeChannel));
 	ChildChannel = MakeUnique<FAgentMessageChannel>(MoveTemp(SecondComputeChannel));
