[UbaCoordinator]
* Fixed so UbaHordeMetaClient is estimating number of cores if LogicalCores is not part of properties array fetched from horde

diff --git a/Engine/Source/Developer/UbaCoordinatorHorde/Private/UbaHordeAgentManager.cpp b/Engine/Source/Developer/UbaCoordinatorHorde/Private/UbaHordeAgentManager.cpp
--- a/Engine/Source/Developer/UbaCoordinatorHorde/Private/UbaHordeAgentManager.cpp
+++ b/Engine/Source/Developer/UbaCoordinatorHorde/Private/UbaHordeAgentManager.cpp
@@ -380,6 +380,8 @@ void FUbaHordeAgentManager::ThreadAgent(FHordeAgentWrapper& Wrapper)
 		ActiveCoreCount += MachineCoreCount;
 	}
 
+	ON_SCOPE_EXIT{ EstimatedCoreCount -= MachineCoreCount; ActiveCoreCount -= MachineCoreCount; };
+
 	uint32 callCounter = 0; // TODO: This should react on the listen string instead of waiting for two text messages :)
 
 	while (Agent->IsValid() && !ShouldExit.Wait(100))
@@ -404,7 +406,4 @@ void FUbaHordeAgentManager::ThreadAgent(FHordeAgentWrapper& Wrapper)
 			}
 		}
 	}
-
-	ActiveCoreCount -= MachineCoreCount;
-	EstimatedCoreCount -= MachineCoreCount;
 }
diff --git a/Engine/Source/Developer/UbaCoordinatorHorde/Private/UbaHordeMetaClient.cpp b/Engine/Source/Developer/UbaCoordinatorHorde/Private/UbaHordeMetaClient.cpp
--- a/Engine/Source/Developer/UbaCoordinatorHorde/Private/UbaHordeMetaClient.cpp
+++ b/Engine/Source/Developer/UbaCoordinatorHorde/Private/UbaHordeMetaClient.cpp
@@ -102,6 +102,9 @@ TSharedPtr<FUbaHordeMetaClient::HordeMachinePromise, ESPMode::ThreadSafe> FUbaHo
 
 			FString OsFamily(TEXT("UNKNOWN-OS"));
 
+			uint16 LogicalCores = 0;
+			uint16 PhysicalCores = 0;
+
 			if (TSharedPtr<FJsonValue> PropertiesValue = OutJson->AsObject()->TryGetField(TEXT("properties")))
 			{
 				for (const TSharedPtr<FJsonValue>& PropertyEntryValue : PropertiesValue->AsArray())
@@ -116,11 +119,28 @@ TSharedPtr<FUbaHordeMetaClient::HordeMachinePromise, ESPMode::ThreadSafe> FUbaHo
 					}
 					if (PropertyElementString.StartsWith(TEXT("LogicalCores=")))
 					{
-						Info.LogicalCores = (uint16)FCString::Atoi(*PropertyElementString + 13);
+						LogicalCores = (uint16)FCString::Atoi(*PropertyElementString + 13);
+					}
+					if (PropertyElementString.StartsWith(TEXT("PhysicalCores=")))
+					{
+						PhysicalCores = (uint16)FCString::Atoi(*PropertyElementString + 14);
 					}
 				}
 			}
 
+			if (LogicalCores)
+			{
+				Info.LogicalCores = LogicalCores;
+			}
+			else if (PhysicalCores)
+			{
+				Info.LogicalCores = PhysicalCores * 2;
+			}
+			else
+			{
+				Info.LogicalCores = 16; // Wild guess
+			}
+
 			FString NonceString = NonceValue->AsString();
 			FString IpString = IpValue->AsString();
 			uint16 PortNumber = (uint16)PortValue->AsNumber();
