[UBA] Add IP and port to warning when UbaAgent connection failed.

#rnx

diff --git a/Engine/Source/Developer/UbaCoordinatorHorde/Private/UbaHordeComputeTransport.cpp b/Engine/Source/Developer/UbaCoordinatorHorde/Private/UbaHordeComputeTransport.cpp
--- a/Engine/Source/Developer/UbaCoordinatorHorde/Private/UbaHordeComputeTransport.cpp
+++ b/Engine/Source/Developer/UbaCoordinatorHorde/Private/UbaHordeComputeTransport.cpp
@@ -19,7 +19,7 @@ FUbaHordeComputeTransport::FUbaHordeComputeTransport(const FHordeRemoteMachineIn
 
 	if (!Socket->Connect(*Address))
 	{
-		UE_LOG(LogUbaHorde, Warning, TEXT("Failed to connect to Horde Agent"));
+		UE_LOG(LogUbaHorde, Warning, TEXT("Failed to connect to Horde Agent [%s:%d]"), *MachineInfo.Ip, (int32)MachineInfo.Port);
 
 		SocketSubsystem->DestroySocket(Socket);
 
