[UBA] Increase verbosity to failed Horde agent connection to Display.

This ensures Horde agent connections failures due to unavailable resources (HTTP 503) are always visible to users.
Also use the EHttpServerResponseCodes values instead of hard coded HTTP constants.

#rnx

diff --git a/Engine/Source/Developer/UbaCoordinatorHorde/Private/UbaHordeMetaClient.cpp b/Engine/Source/Developer/UbaCoordinatorHorde/Private/UbaHordeMetaClient.cpp
--- a/Engine/Source/Developer/UbaCoordinatorHorde/Private/UbaHordeMetaClient.cpp
+++ b/Engine/Source/Developer/UbaCoordinatorHorde/Private/UbaHordeMetaClient.cpp
@@ -6,6 +6,7 @@
 #include "Serialization/JsonSerializer.h"
 #include "Misc/App.h"
 #include "Misc/ConfigCacheIni.h"
+#include "HttpServerConstants.h"
 
 DEFINE_LOG_CATEGORY(LogUbaHorde);
 
@@ -63,16 +64,19 @@ TSharedPtr<FUbaHordeMetaClient::HordeMachinePromise, ESPMode::ThreadSafe> FUbaHo
 
 			if (!bConnectedSuccessfully || !HttpResponse.IsValid())
 			{
-				UE_LOG(LogUbaHorde, Verbose, TEXT("No response from Horde"));
+				UE_LOG(LogUbaHorde, Display, TEXT("No response from Horde"));
 				Promise->SetValue(MakeTuple(HttpResponse, Info));
 				return;
 			}
 
-			FString ResponseStr = HttpResponse->GetContentAsString();
+			const EHttpServerResponseCodes ResponseCode = (EHttpServerResponseCodes)HttpResponse->GetResponseCode();
+			const FString ResponseStr = HttpResponse->GetContentAsString();
 
-			if (HttpResponse->GetResponseCode() == 503) // HTTP 503 Service Unavailable
+			if (ResponseCode == EHttpServerResponseCodes::ServiceUnavail ||
+				ResponseCode == EHttpServerResponseCodes::TooManyRequests)
 			{
-				UE_LOG(LogUbaHorde, Verbose, TEXT("No resources available in Horde (%s)"), *ResponseStr);
+				// Service Unavailable
+				UE_LOG(LogUbaHorde, Display, TEXT("Horde agent request returned with HTTP/%d: %s"), ResponseCode, *ResponseStr);
 				Promise->SetValue(MakeTuple(HttpResponse, Info));
 				return;
 			}
diff --git a/Engine/Source/Developer/UbaCoordinatorHorde/UbaCoordinatorHorde.Build.cs b/Engine/Source/Developer/UbaCoordinatorHorde/UbaCoordinatorHorde.Build.cs
--- a/Engine/Source/Developer/UbaCoordinatorHorde/UbaCoordinatorHorde.Build.cs
+++ b/Engine/Source/Developer/UbaCoordinatorHorde/UbaCoordinatorHorde.Build.cs
@@ -22,6 +22,10 @@ public class UbaCoordinatorHorde : ModuleRules
 			"Sockets",
 		});
 
+		PrivateIncludePathModuleNames.AddRange(new string[] {
+			"HTTPServer",
+		});
+
 		string UbaSourcePath = Path.Combine(EngineDirectory, "Source", "Programs", "UnrealBuildAccelerator");
 		PublicIncludePaths.Add(Path.Combine(UbaSourcePath, "Core", "Public"));
 		PublicIncludePaths.Add(Path.Combine(UbaSourcePath, "Common", "Public"));
