[UBA] Set MaxCores=500 by default in UbaHordeAgentManager.

If UbaController is not used this has no effect and only sets the maximum number of remote agent cores to 500 if this value is not specified externally.
It fixes an issue with parsing INI configurations from the commandline that contain multiple key-value pairs.

#rnx

diff --git a/Engine/Source/Developer/UbaCoordinatorHorde/Public/UbaHordeAgentManager.h b/Engine/Source/Developer/UbaCoordinatorHorde/Public/UbaHordeAgentManager.h
--- a/Engine/Source/Developer/UbaCoordinatorHorde/Public/UbaHordeAgentManager.h
+++ b/Engine/Source/Developer/UbaCoordinatorHorde/Public/UbaHordeAgentManager.h
@@ -48,7 +48,7 @@ private:
 	FString Pool;
 	FAnsiString UbaHost;
 	uint32 UbaPort = 7001;
-	uint32 MaxCores = 0;
+	uint32 MaxCores = 500; // Assume 500 by default in case this configuration is missing but a UBA pool was found
 
 	TUniquePtr<FUbaHordeMetaClient> HordeMetaClient;
 
