[UBA] Lower verbosity of failed UBA/Horde agent connection.

This was requested to reduce noise in CIS builds and this is not considered a warning scenario.

#rnx

diff --git a/Engine/Source/Developer/UbaCoordinatorHorde/Private/UbaHordeComputeTransport.cpp b/Engine/Source/Developer/UbaCoordinatorHorde/Private/UbaHordeComputeTransport.cpp
--- a/Engine/Source/Developer/UbaCoordinatorHorde/Private/UbaHordeComputeTransport.cpp
+++ b/Engine/Source/Developer/UbaCoordinatorHorde/Private/UbaHordeComputeTransport.cpp
@@ -19,7 +19,7 @@ FUbaHordeComputeTransport::FUbaHordeComputeTransport(const FHordeRemoteMachineIn
 
 	if (!Socket->Connect(*Address))
 	{
-		UE_LOG(LogUbaHorde, Warning, TEXT("Failed to connect to Horde Agent [%s:%d]"), *MachineInfo.Ip, (int32)MachineInfo.Port);
+		UE_LOG(LogUbaHorde, Display, TEXT("Failed to connect to Horde Agent [%s:%d]"), *MachineInfo.Ip, (int32)MachineInfo.Port);
 
 		SocketSubsystem->DestroySocket(Socket);
 
