[UBA] Add AES encryption support to UbaCoordinatorHorde module.

This adds support to run UBA/Horde agents over AES encryption for secure network communication in the C++ interface
and can be enabled via the "[UbaController]:Horde" INI configuration, e.g. 'Horde=(Pool="UbaPoolName", MaxCores=500, Encryption=AES)'.
The symmetric AES key is transmitted over HTTPS in the Json response body and provided by the Horde server.
SSL and OpenSSL are added to PrivateDependencyModuleNames in the UbaCoordinateHorde build script.

diff --git a/Engine/Plugins/UbaController/Source/Private/UbaJobProcessor.cpp b/Engine/Plugins/UbaController/Source/Private/UbaJobProcessor.cpp
--- a/Engine/Plugins/UbaController/Source/Private/UbaJobProcessor.cpp
+++ b/Engine/Plugins/UbaController/Source/Private/UbaJobProcessor.cpp
@@ -323,6 +323,23 @@ void FUbaJobProcessor::RunTaskWithUba(FTask* Task)
 
 FString GetUbaBinariesPath();
 
+static bool ParseHordeTransportEncryption(EHordeTransportEncryption& OutEncryption, const TCHAR* InName)
+{
+	if (FCString::Stricmp(InName, TEXT("NONE")) == 0)
+	{
+		OutEncryption = EHordeTransportEncryption::None;
+	}
+	else if (FCString::Stricmp(InName, TEXT("AES")) == 0)
+	{
+		OutEncryption = EHordeTransportEncryption::AES;
+	}
+	else
+	{
+		return false;
+	}
+	return true;
+}
+
 void FUbaJobProcessor::StartUba()
 {
 	checkf(UbaServer == nullptr, TEXT("FUbaJobProcessor::StartUba() was called twice before FUbaJobProcessor::ShutDownUba()"));
@@ -412,6 +429,21 @@ void FUbaJobProcessor::StartUba()
 			UE_LOG(LogUbaController, Log, TEXT("Found UBA controller MaxCores: \"%u\""), MaxCores);
 			HordeAgentManager->SetMaxCoreCount(MaxCores);
 		}
+
+		FString EncryptionName;
+		if (FParse::Value(*HordeConfig, TEXT("Encryption="), EncryptionName))
+		{
+			EHordeTransportEncryption Encryption = EHordeTransportEncryption::None;
+			if (ParseHordeTransportEncryption(Encryption, *EncryptionName))
+			{
+				HordeAgentManager->SetEncryption(Encryption);
+				UE_LOG(LogUbaController, Log, TEXT("Found UBA controller Encryption: \"%s\""), *EncryptionName);
+			}
+			else
+			{
+				UE_LOG(LogUbaController, Log, TEXT("Found UBA controller Encryption with unrecognized identifier: \"%s\""), *EncryptionName);
+			}
+		}
 	}
 
 	FString Host;
diff --git a/Engine/Source/Developer/Horde/Private/Compute/ComputeSocket.cpp b/Engine/Source/Developer/Horde/Private/Compute/ComputeSocket.cpp
--- a/Engine/Source/Developer/Horde/Private/Compute/ComputeSocket.cpp
+++ b/Engine/Source/Developer/Horde/Private/Compute/ComputeSocket.cpp
@@ -259,8 +259,8 @@ public:
 
 	struct FFrameHeader
 	{
-		int Channel;
-		int Size;
+		int32 Channel;
+		int32 Size;
 	};
 
 	TUniquePtr<FComputeTransport> Transport;
@@ -298,8 +298,14 @@ public:
 		}
 
 		Transport->Close();
-		RecvThread.join();
-		PingThread.join();
+
+		// Only join receive and ping threads if they started execution yet
+		if (RecvThread.joinable())
+		{
+			check(PingThread.joinable());
+			RecvThread.join();
+			PingThread.join();
+		}
 	}
 
 	virtual void StartCommunication()
diff --git a/Engine/Source/Developer/Horde/Private/Compute/ComputeTransport.cpp b/Engine/Source/Developer/Horde/Private/Compute/ComputeTransport.cpp
--- a/Engine/Source/Developer/Horde/Private/Compute/ComputeTransport.cpp
+++ b/Engine/Source/Developer/Horde/Private/Compute/ComputeTransport.cpp
@@ -49,6 +49,11 @@ FBufferTransport::FBufferTransport(FComputeBufferWriter InSendBufferWriter, FCom
 {
 }
 
+bool FBufferTransport::IsValid() const
+{
+	return SendBufferWriter.IsValid() && RecvBufferReader.IsValid();
+}
+
 size_t FBufferTransport::Send(const void* Data, size_t Size)
 {
 	unsigned char* Buffer = SendBufferWriter.WaitToWrite(1);
diff --git a/Engine/Source/Developer/Horde/Public/Compute/ComputeTransport.h b/Engine/Source/Developer/Horde/Public/Compute/ComputeTransport.h
--- a/Engine/Source/Developer/Horde/Public/Compute/ComputeTransport.h
+++ b/Engine/Source/Developer/Horde/Public/Compute/ComputeTransport.h
@@ -12,6 +12,9 @@ class FComputeTransport
 {
 public:
 	HORDE_API virtual ~FComputeTransport();
+
+	// Returns true if this compute transport channel has been successfully established.
+	virtual bool IsValid() const = 0;
 	
 	// Sends data to the remote
 	virtual size_t Send(const void* Data, size_t Size) = 0;
@@ -41,6 +44,7 @@ public:
 	HORDE_API FBufferTransport(FComputeBufferWriter InSendBufferWriter, FComputeBufferReader InRecvBufferReader);
 
 protected:
+	virtual bool IsValid() const override final;
 	virtual size_t Send(const void* Data, size_t Size) override final;
 	virtual size_t Recv(void* Data, size_t Size) override final;
 	virtual void MarkComplete() override final;
diff --git a/Engine/Source/Developer/UbaCoordinatorHorde/Private/UbaHordeAgent.cpp b/Engine/Source/Developer/UbaCoordinatorHorde/Private/UbaHordeAgent.cpp
--- a/Engine/Source/Developer/UbaCoordinatorHorde/Private/UbaHordeAgent.cpp
+++ b/Engine/Source/Developer/UbaCoordinatorHorde/Private/UbaHordeAgent.cpp
@@ -8,6 +8,7 @@
 #include "Compute/ComputeSocket.h"
 #include "Compute/ComputeTransport.h"
 #include "UbaHordeComputeTransport.h"
+#include "UbaHordeComputeTransportAES.h"
 
 DEFINE_LOG_CATEGORY(LogUbaHordeAgent);
 
@@ -17,14 +18,25 @@ FUbaHordeAgent::FUbaHordeAgent(const FHordeRemoteMachineInfo& InMachineInfo)
 	, MachineInfo(InMachineInfo)
 {
 	// Create the compute transport object and directly send over the nonce.
-	TUniquePtr<FUbaHordeComputeTransport> ComputeTransport = MakeUnique<FUbaHordeComputeTransport>(MachineInfo, bHasErrors);
+	TUniquePtr<FComputeTransport> ComputeTransport = MakeUnique<FUbaHordeComputeTransport>(MachineInfo, bHasErrors);
 	if (!ComputeTransport->IsValid())
 	{
 		return;
 	}
 
+	// Send initial nonce to establish connection
 	ComputeTransport->Send(MachineInfo.Nonce, sizeof(MachineInfo.Nonce));
 
+	// If encryption protocol is requested, wrap initial transport layer around an encryption layer *after* initial nonce was sent unencrypted
+	if (InMachineInfo.Encryption == EHordeTransportEncryption::AES)
+	{
+		ComputeTransport = MakeUnique<FUbaHordeComputeTransportAES>(MachineInfo, MoveTemp(ComputeTransport));
+		if (!ComputeTransport->IsValid())
+		{
+			return;
+		}
+	}
+
 	// Create the compute socket and initialize a compute channel with the recv/send buffers.
 	// This will allow us to create the agent channel object which will allow us to communicate
 	// with the agent directly.
diff --git a/Engine/Source/Developer/UbaCoordinatorHorde/Private/UbaHordeAgentManager.cpp b/Engine/Source/Developer/UbaCoordinatorHorde/Private/UbaHordeAgentManager.cpp
--- a/Engine/Source/Developer/UbaCoordinatorHorde/Private/UbaHordeAgentManager.cpp
+++ b/Engine/Source/Developer/UbaCoordinatorHorde/Private/UbaHordeAgentManager.cpp
@@ -271,7 +271,7 @@ void FUbaHordeAgentManager::ThreadAgent(FHordeAgentWrapper& Wrapper)
 			}
 		}
 
-		TSharedPtr<FUbaHordeMetaClient::HordeMachinePromise, ESPMode::ThreadSafe> Promise = HordeMetaClient->RequestMachine(Pool);
+		TSharedPtr<FUbaHordeMetaClient::HordeMachinePromise, ESPMode::ThreadSafe> Promise = HordeMetaClient->RequestMachine(Pool, nullptr, Encryption);
 		if (!Promise)
 		{
 			//UE_LOG(LogUbaHorde, Error, TEXT("Failed to create Horde bundle for UbaAgent executable: %s"), *UbaAgentFilePath);
diff --git a/Engine/Source/Developer/UbaCoordinatorHorde/Private/UbaHordeComputeTransport.cpp b/Engine/Source/Developer/UbaCoordinatorHorde/Private/UbaHordeComputeTransport.cpp
--- a/Engine/Source/Developer/UbaCoordinatorHorde/Private/UbaHordeComputeTransport.cpp
+++ b/Engine/Source/Developer/UbaCoordinatorHorde/Private/UbaHordeComputeTransport.cpp
@@ -83,7 +83,7 @@ void FUbaHordeComputeTransport::Close()
 	Socket->Shutdown(ESocketShutdownMode::ReadWrite);
 }
 
-bool FUbaHordeComputeTransport::IsValid()
+bool FUbaHordeComputeTransport::IsValid() const
 {
 	return Socket != nullptr && !bHasErrors && !bIsClosed;
 }
diff --git a/Engine/Source/Developer/UbaCoordinatorHorde/Private/UbaHordeComputeTransport.h b/Engine/Source/Developer/UbaCoordinatorHorde/Private/UbaHordeComputeTransport.h
--- a/Engine/Source/Developer/UbaCoordinatorHorde/Private/UbaHordeComputeTransport.h
+++ b/Engine/Source/Developer/UbaCoordinatorHorde/Private/UbaHordeComputeTransport.h
@@ -6,26 +6,26 @@
 
 class FSocket;
 
-class FUbaHordeComputeTransport : public FComputeTransport
+class FUbaHordeComputeTransport final : public FComputeTransport
 {
 public:
 	// The constructor just performs a connect.
 	FUbaHordeComputeTransport(const FHordeRemoteMachineInfo &Info, bool& InHasErrors);
-	~FUbaHordeComputeTransport() override;
+	~FUbaHordeComputeTransport() override final;
 	
 	// Sends data to the remote
-	size_t Send(const void* Data, size_t Size) override;
+	size_t Send(const void* Data, size_t Size) override final;
 
 	// Receives data from the remote
-	size_t Recv(void* Data, size_t Size) override;
+	size_t Recv(void* Data, size_t Size) override final;
 
 	// Indicates to the remote that no more data will be sent.
-	void MarkComplete() override;
+	void MarkComplete() override final;
 
 	// Indicates that no more data will be sent or received, and that any blocking reads/writes should stop.
-	void Close() override;
+	void Close() override final;
 
-	bool IsValid();
+	bool IsValid() const override final;
 
 private:
 	FSocket* Socket;
diff --git a/Engine/Source/Developer/UbaCoordinatorHorde/Private/UbaHordeComputeTransportAES.cpp b/Engine/Source/Developer/UbaCoordinatorHorde/Private/UbaHordeComputeTransportAES.cpp
new file mode 100644
--- /dev/null
+++ b/Engine/Source/Developer/UbaCoordinatorHorde/Private/UbaHordeComputeTransportAES.cpp
@@ -0,0 +1,339 @@
+// Copyright Epic Games, Inc. All Rights Reserved.
+
+#if WITH_SSL // OpenSSL library required for AES encryption
+
+#include "UbaHordeComputeTransportAES.h"
+
+#include <openssl/evp.h>
+#include <openssl/aes.h>
+
+
+// Message format: [length (4 bytes)][nonce (12 bytes)][encrypted data][tag (16 bytes)]
+namespace CryptoFormatAES
+{
+	static constexpr int32 KeySize = HORDE_KEY_SIZE;
+	static constexpr int32 BufferLengthSize = sizeof(int32);
+	static constexpr int32 NonceSize = 12;
+	static constexpr int32 NonceWords = NonceSize / sizeof(int32);
+	static constexpr int32 TagSize = AES_BLOCK_SIZE; // AES-GCM tag size in bytes
+
+	struct FPacketHeader
+	{
+		int32 DataLength = 0;
+		uint8 Nonce[CryptoFormatAES::NonceSize] = {};
+	};
+
+	static constexpr int32 HeaderAndTagSize = sizeof(FPacketHeader) + TagSize;
+} // AESFormat
+
+struct FUbaHordeComputeTransportAES::FOpenSSLContext
+{
+	EVP_CIPHER_CTX* EncryptionContext = nullptr;
+	EVP_CIPHER_CTX* DecryptionContext = nullptr;
+
+	uint8 Key[CryptoFormatAES::KeySize] = {};
+	uint8 EncryptNonce[CryptoFormatAES::NonceSize] = {};
+	uint8 DecryptNonce[CryptoFormatAES::NonceSize] = {};
+
+	FOpenSSLContext(const uint8 (&InKey)[HORDE_KEY_SIZE])
+	{
+		FMemory::Memcpy(Key, InKey, sizeof(InKey));
+
+		// Initialize nonce for encryption with random numbers
+		for (uint8& NonceWord : EncryptNonce)
+		{
+			NonceWord = (uint8)(FMath::Rand() % 256);
+		}
+	}
+
+	void UpdateNonce()
+	{
+		// Mangle nonce to add non-determinism to message digest
+		uint32* Nonce32Bits = reinterpret_cast<uint32*>(EncryptNonce);
+		Nonce32Bits[0]++;
+		Nonce32Bits[1]--;
+		Nonce32Bits[2] = Nonce32Bits[0] ^ Nonce32Bits[1];
+	}
+
+	void EncryptBegin()
+	{
+		// Update nonce to disguise equal messages with non-determinism
+		UpdateNonce();
+
+		// Use nonce as initialization vector
+		if (EVP_EncryptInit_ex(EncryptionContext, nullptr, nullptr, Key, EncryptNonce) == 0)
+		{
+			UE_LOG(LogUbaHorde, Verbose, TEXT("EVP_EncryptInit_ex() failed"));
+		}
+	}
+
+	int32 Copy(void* OutData, const void* InData, int32 InDataLength)
+	{
+		FMemory::Memcpy(OutData, InData, (SIZE_T)InDataLength);
+		return InDataLength;
+	}
+
+	int32 Encrypt(void* OutDataCipher, const void* InData, int32 InDataLength)
+	{
+		int OutDataLength = 0;
+		if (EVP_EncryptUpdate(EncryptionContext, (unsigned char*)OutDataCipher, &OutDataLength, (const unsigned char*)InData, InDataLength) == 0)
+		{
+			UE_LOG(LogUbaHorde, Verbose, TEXT("EVP_EncryptUpdate(Size = %d) failed"), InDataLength);
+		}
+		return (int32)OutDataLength;
+	}
+
+	int32 EncryptEnd(void* OutData)
+	{
+		int OutDataLength = 0;
+
+		int FooterLength = 0;
+		if (EVP_EncryptFinal_ex(EncryptionContext, (unsigned char*)OutData, &FooterLength) == 0)
+		{
+			UE_LOG(LogUbaHorde, Verbose, TEXT("EVP_EncryptFinal_ex() failed"));
+		}
+		OutDataLength += FooterLength;
+
+		if (EVP_CIPHER_CTX_ctrl(EncryptionContext, EVP_CTRL_GCM_GET_TAG, CryptoFormatAES::TagSize, ((char*)OutData) + OutDataLength) == 0)
+		{
+			UE_LOG(LogUbaHorde, Verbose, TEXT("EVP_CIPHER_CTX_ctrl(EVP_CTRL_GCM_GET_TAG) failed"));
+		}
+		OutDataLength += CryptoFormatAES::TagSize;
+
+		return (int32)OutDataLength;
+	}
+
+	int32 EncryptMessage(void* OutEncryptedData, const void* InData, int32 InDataLength)
+	{
+		int32 EncryptedDataLength = 0;
+
+		EncryptBegin();
+
+		uint8* EncryptedBytes = (uint8*)OutEncryptedData;
+		EncryptedDataLength += Copy(EncryptedBytes, &InDataLength, sizeof(InDataLength));
+		EncryptedDataLength += Copy(EncryptedBytes + EncryptedDataLength, EncryptNonce, sizeof(EncryptNonce));
+		EncryptedDataLength += Encrypt(EncryptedBytes + EncryptedDataLength, InData, InDataLength);
+		EncryptedDataLength += EncryptEnd(EncryptedBytes + EncryptedDataLength);
+
+		return EncryptedDataLength;
+	}
+
+	void DecryptBegin(const void* InNonce, int32 InNonceSize)
+	{
+		// Use nonce that was received with the encrypted message
+		check(InNonce != nullptr);
+		check(InNonceSize == CryptoFormatAES::NonceSize);
+		FMemory::Memcpy(DecryptNonce, InNonce, (SIZE_T)InNonceSize);
+
+		// Use nonce as initialization vector
+		if (EVP_DecryptInit_ex(DecryptionContext, nullptr, nullptr, Key, DecryptNonce) == 0)
+		{
+			UE_LOG(LogUbaHorde, Verbose, TEXT("EVP_DecryptInit_ex() failed"));
+		}
+	}
+
+	int32 Decrypt(void* OutDataPlaintext, const void* InData, int32 InDataLength)
+	{
+		int OutDataLength = 0;
+		if (EVP_DecryptUpdate(DecryptionContext, (unsigned char*)OutDataPlaintext, &OutDataLength, (const unsigned char*)InData, InDataLength) == 0)
+		{
+			UE_LOG(LogUbaHorde, Verbose, TEXT("EVP_DecryptUpdate(Size = %d) failed"), InDataLength);
+		}
+		return (int32)OutDataLength;
+	}
+
+	int32 DecryptEnd(void* OutDataPlaintext)
+	{
+		int OutDataLength = 0;
+		unsigned char* OutBytes = (unsigned char*)OutDataPlaintext;
+
+		unsigned char Tag[CryptoFormatAES::TagSize];
+		FMemory::Memcpy(Tag, OutBytes, sizeof(Tag));
+		if (EVP_CIPHER_CTX_ctrl(DecryptionContext, EVP_CTRL_GCM_SET_TAG, CryptoFormatAES::TagSize, Tag) == 0)
+		{
+			UE_LOG(LogUbaHorde, Verbose, TEXT("EVP_CIPHER_CTX_ctrl(EVP_CTRL_GCM_SET_TAG) failed"));
+		}
+		OutDataLength += CryptoFormatAES::TagSize;
+
+		int FooterLength = 0;
+		if (EVP_DecryptFinal_ex(DecryptionContext, OutBytes + OutDataLength, &FooterLength) == 0)
+		{
+			UE_LOG(LogUbaHorde, Verbose, TEXT("EVP_DecryptFinal_ex() failed"));
+		}
+		OutDataLength += FooterLength;
+
+		return (int32)OutDataLength;
+	}
+};
+
+static EVP_CIPHER_CTX* NewAESCryptoContext()
+{
+	EVP_CIPHER_CTX* CryptoContext = EVP_CIPHER_CTX_new();
+	EVP_CIPHER_CTX_init(CryptoContext);
+	EVP_CIPHER_CTX_set_padding(CryptoContext, 0);
+	return CryptoContext;
+}
+
+FUbaHordeComputeTransportAES::FUbaHordeComputeTransportAES(const FHordeRemoteMachineInfo& MachineInfo, TUniquePtr<FComputeTransport>&& InInnerTransport)
+	: OpenSSLContext(new FOpenSSLContext(MachineInfo.Key))
+	, InnerTransport(MoveTemp(InInnerTransport))
+	, RemainingDataOffset(0)
+	, bIsClosed(false)
+{
+	// Initialize AES encryption and decrption contexts
+	OpenSSLContext->EncryptionContext = NewAESCryptoContext();
+	EVP_EncryptInit_ex(OpenSSLContext->EncryptionContext, EVP_aes_256_gcm(), nullptr, nullptr, nullptr);
+
+	OpenSSLContext->DecryptionContext = NewAESCryptoContext();
+	EVP_DecryptInit_ex(OpenSSLContext->DecryptionContext, EVP_aes_256_gcm(), nullptr, nullptr, nullptr);
+}
+
+FUbaHordeComputeTransportAES::~FUbaHordeComputeTransportAES()
+{
+	if (OpenSSLContext)
+	{
+		if (OpenSSLContext->EncryptionContext)
+		{
+			EVP_CIPHER_CTX_free(OpenSSLContext->EncryptionContext);
+		}
+		if (OpenSSLContext->DecryptionContext)
+		{
+			EVP_CIPHER_CTX_free(OpenSSLContext->DecryptionContext);
+		}
+		delete OpenSSLContext;
+	}
+}
+
+// Sends data to the remote
+size_t FUbaHordeComputeTransportAES::Send(const void* Data, size_t Size)
+{
+	if (!IsValid())
+	{
+		return 0;
+	}
+
+	FScopeLock LockGuard(&IntermediateBuffersLock);
+
+	// Encrypt data before sending
+	int32 DataLength = (int32)Size;
+	int32 MessageLength = CryptoFormatAES::HeaderAndTagSize + DataLength;
+	uint8* EncryptedData = GetAndResizeEncryptedBuffer(MessageLength);
+
+	int32 EncryptedDataLength = OpenSSLContext->EncryptMessage(EncryptedData, Data, DataLength);
+
+	// Send encrypted data over common transport layer
+	if (!InnerTransport->SendMessage(EncryptedData, (size_t)EncryptedDataLength))
+	{
+		return 0;
+	}
+
+	return DataLength;
+}
+
+// Receives data from the remote
+size_t FUbaHordeComputeTransportAES::Recv(void* Data, size_t Size)
+{
+	if (!IsValid())
+	{
+		return 0;
+	}
+
+	// Check if there is remaining decrypted data from previous call
+	if (!RemainingData.IsEmpty())
+	{
+		FScopeLock LockGuard(&IntermediateBuffersLock);
+
+		const int32 NumBytesFromRemainingData = FMath::Min(RemainingData.Num() - RemainingDataOffset, (int32)Size);
+
+		// Return remaining data from intermediate buffer
+		FMemory::Memcpy(Data, &RemainingData[RemainingDataOffset], (size_t)NumBytesFromRemainingData);
+
+		// When all remaining data has been consumed, reset intermediate buffer
+		RemainingDataOffset += NumBytesFromRemainingData;
+
+		if (RemainingDataOffset >= RemainingData.Num())
+		{
+			RemainingData.Empty();
+		}
+
+		return (size_t)NumBytesFromRemainingData;
+	}
+
+	// Receive message header
+	CryptoFormatAES::FPacketHeader Header;
+
+	if (!InnerTransport->RecvMessage(&Header, sizeof(Header)))
+	{
+		UE_LOG(LogUbaHorde, VeryVerbose, TEXT("Failed to receive AES packet (Expected size = %d)"), (int32)Size);
+		return 0;
+	}
+
+	FScopeLock LockGuard(&IntermediateBuffersLock);
+
+	// Recieve message data
+	int32 MessageLength = CryptoFormatAES::TagSize + Header.DataLength;
+	uint8* EncryptedData = GetAndResizeEncryptedBuffer(MessageLength);
+
+	if (!InnerTransport->RecvMessage(EncryptedData, (size_t)MessageLength))
+	{
+		UE_LOG(LogUbaHorde, VeryVerbose, TEXT("Failed to receive AES packet (Expected size = %d, Size specified in header = %d)"), (int32)Size, Header.DataLength);
+		return 0;
+	}
+
+	OpenSSLContext->DecryptBegin(Header.Nonce, sizeof(Header.Nonce));
+
+	// Decrypt data and store remaining encrypted data for next Recv() call
+	const int32 NumBytesProcessed = FMath::Min(Header.DataLength, (int32)Size);
+
+	OpenSSLContext->Decrypt(Data, EncryptedData, NumBytesProcessed);
+
+	if (Header.DataLength > NumBytesProcessed)
+	{
+		RemainingDataOffset = 0;
+		RemainingData.SetNumUninitialized(Header.DataLength - NumBytesProcessed);
+		OpenSSLContext->Decrypt(RemainingData.GetData(), EncryptedData + NumBytesProcessed, RemainingData.Num());
+	}
+
+	OpenSSLContext->DecryptEnd(EncryptedData + Header.DataLength);
+
+	return NumBytesProcessed;
+}
+
+// Indicates to the remote that no more data will be sent.
+void FUbaHordeComputeTransportAES::MarkComplete()
+{
+	if (IsValid())
+	{
+		InnerTransport->MarkComplete();
+	}
+}
+
+// Indicates that no more data will be sent or received, and that any blocking reads/writes should stop.
+void FUbaHordeComputeTransportAES::Close()
+{
+	if (!bIsClosed)
+	{
+		if (IsValid())
+		{
+			InnerTransport->Close();
+		}
+		bIsClosed = true;
+	}
+}
+
+bool FUbaHordeComputeTransportAES::IsValid() const
+{
+	return InnerTransport.IsValid() && InnerTransport->IsValid();
+}
+
+uint8* FUbaHordeComputeTransportAES::GetAndResizeEncryptedBuffer(int32 Size)
+{
+	if (Size > EncryptedBuffer.Num())
+	{
+		EncryptedBuffer.SetNumUninitialized(Size);
+	}
+	return EncryptedBuffer.GetData();
+}
+
+#endif // WITH_SSL
+
diff --git a/Engine/Source/Developer/UbaCoordinatorHorde/Private/UbaHordeComputeTransportAES.h b/Engine/Source/Developer/UbaCoordinatorHorde/Private/UbaHordeComputeTransportAES.h
new file mode 100644
--- /dev/null
+++ b/Engine/Source/Developer/UbaCoordinatorHorde/Private/UbaHordeComputeTransportAES.h
@@ -0,0 +1,45 @@
+// Copyright Epic Games, Inc. All Rights Reserved.
+#pragma once
+
+#if WITH_SSL
+
+#include "Compute/ComputeTransport.h"
+#include "UbaHordeMetaClient.h"
+#include "UbaHordeComputeTransport.h"
+
+class FUbaHordeComputeTransportAES final : public FComputeTransport
+{
+public:
+	// The constructor just performs a connect.
+	FUbaHordeComputeTransportAES(const FHordeRemoteMachineInfo &Info, TUniquePtr<FComputeTransport>&& InInnerTransport);
+	~FUbaHordeComputeTransportAES() override final;
+	
+	// Sends data to the remote
+	size_t Send(const void* Data, size_t Size) override final;
+
+	// Receives data from the remote
+	size_t Recv(void* Data, size_t Size) override final;
+
+	// Indicates to the remote that no more data will be sent.
+	void MarkComplete() override final;
+
+	// Indicates that no more data will be sent or received, and that any blocking reads/writes should stop.
+	void Close() override final;
+
+	bool IsValid() const override final;
+
+private:
+	struct FOpenSSLContext;
+
+	uint8* GetAndResizeEncryptedBuffer(int32 Size);
+
+	FOpenSSLContext* OpenSSLContext;
+	TUniquePtr<FComputeTransport> InnerTransport;
+	TArray<uint8> EncryptedBuffer;
+	TArray<uint8> RemainingData;
+	int32 RemainingDataOffset;
+	FCriticalSection IntermediateBuffersLock;
+	bool bIsClosed;
+};
+
+#endif // WITH_SSL
diff --git a/Engine/Source/Developer/UbaCoordinatorHorde/Private/UbaHordeMetaClient.cpp b/Engine/Source/Developer/UbaCoordinatorHorde/Private/UbaHordeMetaClient.cpp
--- a/Engine/Source/Developer/UbaCoordinatorHorde/Private/UbaHordeMetaClient.cpp
+++ b/Engine/Source/Developer/UbaCoordinatorHorde/Private/UbaHordeMetaClient.cpp
@@ -35,11 +35,37 @@ bool FUbaHordeMetaClient::RefreshHttpClient()
 	return true;
 }
 
-TSharedPtr<FUbaHordeMetaClient::HordeMachinePromise, ESPMode::ThreadSafe> FUbaHordeMetaClient::RequestMachine(const FString& PoolId, const FString& Machine)
+static const TCHAR* LexToString(EHordeTransportEncryption InTransportEncrypted)
+{
+	switch (InTransportEncrypted)
+	{
+	case EHordeTransportEncryption::None:	break;
+	case EHordeTransportEncryption::AES:	return TEXT("aes");
+	}
+	return TEXT("");
+}
+
+static void BuildHordeRequestJsonBody(TStringBuilder<1024>& RequestJsonBody, const FString& InPoolId, EHordeTransportEncryption InTransportEncrypted)
+{
+	RequestJsonBody.AppendChar(TEXT('{'));
+
+	// Append requirements: pool ID and whether it should be allocated exclusively for this request
+	RequestJsonBody.Appendf(TEXT("\"requirements\":{\"pool\":\"%s\",\"exclusive\":true}"), *InPoolId);
+
+	// Append connection information if encryption is enabled
+	if (InTransportEncrypted != EHordeTransportEncryption::None)
+	{
+		RequestJsonBody.Appendf(TEXT(",\"connection\":{\"encryption\":\"%s\",\"modePreference\":\"Direct\"}"), LexToString(InTransportEncrypted));
+	}
+
+	RequestJsonBody.AppendChar(TEXT('}'));
+}
+
+TSharedPtr<FUbaHordeMetaClient::HordeMachinePromise, ESPMode::ThreadSafe> FUbaHordeMetaClient::RequestMachine(const FString& PoolId, const TCHAR* Machine, EHordeTransportEncryption Encryption)
 {
 	TSharedPtr<HordeMachinePromise> Promise = MakeShared<HordeMachinePromise, ESPMode::ThreadSafe>();
 
-	const FString ResourcePath = FString::Format(TEXT("api/v2/compute/{0}"), { Machine });
+	const FString ResourcePath = FString::Format(TEXT("api/v2/compute/{0}"), { Machine != nullptr ? Machine : TEXT("default") });
 
 	FHttpRequestRef Request = HttpClient->CreateRequest(TEXT("POST"), *ResourcePath);
 
@@ -47,14 +73,16 @@ TSharedPtr<FUbaHordeMetaClient::HordeMachinePromise, ESPMode::ThreadSafe> FUbaHo
 
 	// Post JSON document with constraints to acquire a Horde agent. Use pool of agents (e.g. "BoxLinux" or "BoxWin")
 	// and require exclusive access or UbaStorage will fail to initialize the next time we connect to the same machine.
-	const FString Body = FString::Format(TEXT("{\"requirements\":{\"pool\":\"{0}\",\"exclusive\":true}}"), { PoolId });
-	Request->SetContentAsString(Body);
+	TStringBuilder<1024> RequestJsonBody;
+	BuildHordeRequestJsonBody(RequestJsonBody, PoolId, Encryption);
+
+	Request->SetContentAsString(RequestJsonBody.ToString());
 	Request->SetHeader(TEXT("Content-Type"), TEXT("application/json"));
 
-	UE_LOG(LogUbaHorde, Log, TEXT("Requesting Horde agent with JSON descriptor: '%s'"), *Body);
+	UE_LOG(LogUbaHorde, Log, TEXT("Requesting Horde agent with JSON descriptor: '%s'"), RequestJsonBody.ToString());
 
 	Request->OnProcessRequestComplete().BindLambda(
-		[this, Promise](FHttpRequestPtr /*Request*/, FHttpResponsePtr HttpResponse, bool bConnectedSuccessfully)
+		[this, Promise, Encryption](FHttpRequestPtr /*Request*/, FHttpResponsePtr HttpResponse, bool bConnectedSuccessfully)
 		{
 			FHordeRemoteMachineInfo Info;
 			Info.Ip = TEXT("");
@@ -119,7 +147,9 @@ TSharedPtr<FUbaHordeMetaClient::HordeMachinePromise, ESPMode::ThreadSafe> FUbaHo
 					{
 						OsFamily = *PropertyElementString + 9;
 						if (OsFamily == TEXT("Windows"))
+						{
 							Info.bRunsWindowOS = true;
+						}
 					}
 					if (PropertyElementString.StartsWith(TEXT("LogicalCores=")))
 					{
@@ -145,22 +175,44 @@ TSharedPtr<FUbaHordeMetaClient::HordeMachinePromise, ESPMode::ThreadSafe> FUbaHo
 				Info.LogicalCores = 16; // Wild guess
 			}
 
+			// Log summary of assigned Horde machine
 			FString NonceString = NonceValue->AsString();
 			FString IpString = IpValue->AsString();
 			uint16 PortNumber = (uint16)PortValue->AsNumber();
 
-			if (TSharedPtr<FJsonValue> LeaseIdValue = OutJson->AsObject()->TryGetField(TEXT("leaseId")))
+			FString HordeMachineSummary = FString::Printf(TEXT("UBA Horde machine assigned (%s) [%s:%u]"), *OsFamily, *IpString, (uint32)PortNumber);
+
+#if WITH_SSL
+			if (Encryption == EHordeTransportEncryption::AES)
 			{
-				const FString AgentWebPortalUrl = FString::Format(TEXT("{0}lease/{1}"), { this->ServerUrl, LeaseIdValue->AsString() });
-				UE_LOG(
-					LogUbaHorde, Display, TEXT("UBA Horde machine assigned (%s) [%s:%u]: %s"), *OsFamily, *IpString, (uint32)PortNumber, *AgentWebPortalUrl
-				);
+				TSharedPtr<FJsonValue> KeyValue = OutJson->AsObject()->TryGetField(TEXT("key"));
+				FString KeyString = KeyValue ? KeyValue->AsString() : FString();
+				if (KeyString.IsEmpty())
+				{
+					HordeMachineSummary += TEXT(" [AES key missing]");
+				}
+				else if (KeyString.Len() != sizeof(Info.Key) * 2)
+				{
+					HordeMachineSummary += TEXT(" [AES key corrupted]");
+				}
+				else
+				{
+					FString::ToHexBlob(KeyString, Info.Key, sizeof(Info.Key));
+					HordeMachineSummary += TEXT(" [AES key received]");
+				}
 			}
-			else
+
+			Info.Encryption = Encryption;
+#endif
+
+			if (TSharedPtr<FJsonValue> LeaseIdValue = OutJson->AsObject()->TryGetField(TEXT("leaseId")))
 			{
-				UE_LOG(LogUbaHorde, Display, TEXT("UBA Horde machine assigned [%s:%u]"), *IpString, (uint32)PortNumber);
+				HordeMachineSummary += FString::Format(TEXT(": {0}lease/{1}"), { this->ServerUrl, LeaseIdValue->AsString() });
 			}
 
+			UE_LOG(LogUbaHorde, Display, TEXT("%s"), *HordeMachineSummary);
+
+			// Return final response information
 			Info.Ip = IpString;
 			Info.Port = PortNumber;
 			FString::ToHexBlob(NonceString, Info.Nonce, HORDE_NONCE_SIZE);
diff --git a/Engine/Source/Developer/UbaCoordinatorHorde/Private/UbaHordeMetaClient.h b/Engine/Source/Developer/UbaCoordinatorHorde/Private/UbaHordeMetaClient.h
--- a/Engine/Source/Developer/UbaCoordinatorHorde/Private/UbaHordeMetaClient.h
+++ b/Engine/Source/Developer/UbaCoordinatorHorde/Private/UbaHordeMetaClient.h
@@ -7,10 +7,12 @@
 #include "Interfaces/IHttpRequest.h"
 #include "Interfaces/IHttpResponse.h"
 #include "HordeHttpClient.h"
+#include "UbaHordeAgentManager.h"
 
 DECLARE_LOG_CATEGORY_EXTERN(LogUbaHorde, Log, Log);
 
 #define HORDE_NONCE_SIZE 64
+#define HORDE_KEY_SIZE 32
 
 // When asking the Horde server for a machine, this will be returned when the HTTP
 // reponse comes back (after a machine request). If the response didn't come through
@@ -21,10 +23,13 @@ DECLARE_LOG_CATEGORY_EXTERN(LogUbaHorde, Log, Log);
 struct FHordeRemoteMachineInfo
 {
 	FString Ip;
-	uint16 Port;
-	uint16 LogicalCores;
-	uint8 Nonce[HORDE_NONCE_SIZE];
-	bool bRunsWindowOS;
+	uint16 Port = 0;
+	uint16 LogicalCores = 0;
+	EHordeTransportEncryption Encryption = EHordeTransportEncryption::None;
+	uint8 Nonce[HORDE_NONCE_SIZE] = {};
+	uint8 Key[HORDE_KEY_SIZE] = {};
+	TArray<uint8> CertificateData; // X.509 certificate for SSL connections in Pkcs12 format.
+	bool bRunsWindowOS = false;
 };
 
 // This encapsulates the mechanism of talking to the Horde "meta server". The server which
@@ -47,7 +52,7 @@ public:
 	// ... // Can to asynchronous work
 	// auto Future.Wait(); // Wait for the response to arrive
 	// FHordeRemoteMachineInfo MachineInfo = Future.Get().Value;
-	TSharedPtr<HordeMachinePromise, ESPMode::ThreadSafe> RequestMachine(const FString& PoolId, const FString& Machine = "default");
+	TSharedPtr<HordeMachinePromise, ESPMode::ThreadSafe> RequestMachine(const FString& PoolId, const TCHAR* Machine = nullptr, EHordeTransportEncryption Encryption = EHordeTransportEncryption::None);
 
 private:
 	void ParseConfig(const FString& IniFilename, const FString& IniSection);
diff --git a/Engine/Source/Developer/UbaCoordinatorHorde/Public/UbaHordeAgentManager.h b/Engine/Source/Developer/UbaCoordinatorHorde/Public/UbaHordeAgentManager.h
--- a/Engine/Source/Developer/UbaCoordinatorHorde/Public/UbaHordeAgentManager.h
+++ b/Engine/Source/Developer/UbaCoordinatorHorde/Public/UbaHordeAgentManager.h
@@ -10,16 +10,23 @@ class FEvent;
 class FUbaHordeMetaClient;
 namespace uba { class NetworkServer; }
 
+enum class EHordeTransportEncryption
+{
+	None,
+	AES,
+};
+
 class FUbaHordeAgentManager
 {
 public:
 	UBACOORDINATORHORDE_API FUbaHordeAgentManager(const FString& InWorkingDir, const FString& BinariesPath);
 	UBACOORDINATORHORDE_API ~FUbaHordeAgentManager();
 
-	inline void SetPool(const FString& InPool) { Pool = InPool; }
-	inline void SetMaxCoreCount(uint32 Count) { MaxCores = Count; }
-	inline void SetUbaHost(const FAnsiString& InHost) { UbaHost = InHost; }
-	inline void SetUbaPort(uint32 InPort) { UbaPort = InPort; }
+	void SetPool(const FString& InPool) { Pool = InPool; }
+	void SetMaxCoreCount(uint32 Count) { MaxCores = Count; }
+	void SetEncryption(EHordeTransportEncryption InEncryption) { Encryption = InEncryption; }
+	void SetUbaHost(const FAnsiString& InHost) { UbaHost = InHost; }
+	void SetUbaPort(uint32 InPort) { UbaPort = InPort; }
 
 	UBACOORDINATORHORDE_API void SetTargetCoreCount(uint32 Count);
 
@@ -49,6 +56,7 @@ private:
 	FAnsiString UbaHost;
 	uint32 UbaPort = 7001;
 	uint32 MaxCores = 500; // Assume 500 by default in case this configuration is missing but a UBA pool was found
+	EHordeTransportEncryption Encryption = EHordeTransportEncryption::None;
 
 	TUniquePtr<FUbaHordeMetaClient> HordeMetaClient;
 
diff --git a/Engine/Source/Developer/UbaCoordinatorHorde/UbaCoordinatorHorde.Build.cs b/Engine/Source/Developer/UbaCoordinatorHorde/UbaCoordinatorHorde.Build.cs
--- a/Engine/Source/Developer/UbaCoordinatorHorde/UbaCoordinatorHorde.Build.cs
+++ b/Engine/Source/Developer/UbaCoordinatorHorde/UbaCoordinatorHorde.Build.cs
@@ -20,6 +20,8 @@ public class UbaCoordinatorHorde : ModuleRules
 			"HTTP",
 			"Json",
 			"Sockets",
+			"SSL",
+			"OpenSSL",
 		});
 
 		PrivateIncludePathModuleNames.AddRange(new string[] {
