[UBA] Refactored UBA configuration and added support for Relay mode for shader compilation.

This contains two significant changes:
1. Refactored how UBA/Horde configuration is read in C++ for shader compilation to match C# code for UBT. This is implemented in FUbaHordeConfig and applies the same list of BuildConfiguration.xml files in the same order as UBT. It also adds the same command line arguments. This effectively replaces the old [UbaController] configurations from BaseEngine.ini. If a machine is set up for UBA in UBT, it will now automatically be set up for UBA in shader compilation.
2. Add support for Relay mode. This required the addition of "_auto" Horde cluster to be resolved via an HTTP/Json request prior to the agent assignment request. This allows the client to get a Horde cluster assigned automatically based on the user's IP address (i.e. region based).

diff --git a/Engine/Plugins/UbaController/Source/Private/UbaControllerModule.cpp b/Engine/Plugins/UbaController/Source/Private/UbaControllerModule.cpp
--- a/Engine/Plugins/UbaController/Source/Private/UbaControllerModule.cpp
+++ b/Engine/Plugins/UbaController/Source/Private/UbaControllerModule.cpp
@@ -4,6 +4,7 @@
 #include "HttpModule.h"
 #include "Misc/ConfigCacheIni.h"
 #include "UbaJobProcessor.h"
+#include "UbaHordeConfig.h"
 #include "Features/IModularFeatures.h"
 #include "HAL/FileManager.h"
 #include "Misc/CommandLine.h"
@@ -90,32 +91,7 @@ static bool IsUbaControllerEnabled()
 	{
 		return false;
 	}
-
-	// Check if UbaController is enabled via command line argument
-	if (FParse::Param(FCommandLine::Get(), TEXT("Uba")))
-	{
-		return true;
-	}
-
-	// Check if UbaController is enabled via INI configuration in [UbaController] section.
-	FString EnabledState;
-	GConfig->GetString(TEXT("UbaController"), TEXT("Enabled"), EnabledState, GEngineIni);
-
-	// This "Enabled" parameter is a tri-state so we have to parse it as a string.
-	// Those strings can include the comments (starting with ';') from the INI file, so we have to trim that section from the string.
-	const int32 EnabledStateCommentPosition = EnabledState.Find(TEXT(";"));
-	if (EnabledStateCommentPosition != INDEX_NONE)
-	{
-		EnabledState.RemoveAt(EnabledStateCommentPosition, EnabledState.Len() - EnabledStateCommentPosition);
-		EnabledState.RemoveSpacesInline();
-	}
-
-	if (EnabledState.Equals(TEXT("True"), ESearchCase::IgnoreCase) || (EnabledState.Equals(TEXT("BuildMachineOnly"), ESearchCase::IgnoreCase) && GIsBuildMachine))
-	{
-		return true;
-	}
-
-	return false;
+	return true;
 }
 
 bool FUbaControllerModule::IsSupported()
diff --git a/Engine/Plugins/UbaController/Source/Private/UbaJobProcessor.cpp b/Engine/Plugins/UbaController/Source/Private/UbaJobProcessor.cpp
--- a/Engine/Plugins/UbaController/Source/Private/UbaJobProcessor.cpp
+++ b/Engine/Plugins/UbaController/Source/Private/UbaJobProcessor.cpp
@@ -8,6 +8,7 @@
 #include "Misc/CoreMisc.h"
 #include "UbaControllerModule.h"
 #include "UbaHordeAgentManager.h"
+#include "UbaHordeConfig.h"
 #include "UbaProcessStartInfo.h"
 #include "UbaSessionServerCreateInfo.h"
 #include "UbaStringConversion.h"
@@ -154,11 +155,7 @@ FUbaJobProcessor::FUbaJobProcessor(
 			checkf(false, TEXT("%s"), UBASTRING_TO_TCHAR(text));
 		});
 
-	if (!GConfig->GetInt(TEXT("UbaController"), TEXT("MaxLocalParallelJobs"), MaxLocalParallelJobs, GEngineIni))
-	{
-		MaxLocalParallelJobs = -1;
-	}
-
+	MaxLocalParallelJobs = FUbaHordeConfig::Get().HordeMaxLocalCores;
 	if (MaxLocalParallelJobs == -1)
 	{
 		MaxLocalParallelJobs = FPlatformMisc::NumberOfCoresIncludingHyperthreads();
@@ -323,23 +320,6 @@ void FUbaJobProcessor::RunTaskWithUba(FDistributedBuildTask* Task)
 
 FString GetUbaBinariesPath();
 
-static bool ParseHordeTransportEncryption(EHordeTransportEncryption& OutEncryption, const TCHAR* InName)
-{
-	if (FCString::Stricmp(InName, TEXT("NONE")) == 0)
-	{
-		OutEncryption = EHordeTransportEncryption::None;
-	}
-	else if (FCString::Stricmp(InName, TEXT("AES")) == 0)
-	{
-		OutEncryption = EHordeTransportEncryption::AES;
-	}
-	else
-	{
-		return false;
-	}
-	return true;
-}
-
 void FUbaJobProcessor::StartUba()
 {
 	checkf(UbaServer == nullptr, TEXT("FUbaJobProcessor::StartUba() was called twice before FUbaJobProcessor::ShutDownUba()"));
@@ -389,80 +369,16 @@ void FUbaJobProcessor::StartUba()
 		NetworkServer_StartListen(UbaServer, uba::DefaultPort, nullptr); // Start listen so any helper on the LAN can join in
 	}
 
-	HordeAgentManager = MakeUnique<FUbaHordeAgentManager>(ControllerModule.GetWorkingDirectory(), GetUbaBinariesPath());
-
-	auto AddClientCallback = [](void* userData, const uba::tchar* ip, uint16 port)
-		{
-			return NetworkServer_AddClient((uba::NetworkServer*)userData, ip, port, nullptr);
-		};
-	HordeAgentManager->SetAddClientCallback(AddClientCallback, UbaServer);
-
-	FString HordeConfig;
-	if (GConfig->GetString(TEXT("UbaController"), TEXT("Horde"), HordeConfig, GEngineIni))
+	// Only request Horde agents if Horde is enabled for UBA
+	if (FUbaHordeConfig::Get().IsHordeEnabled())
 	{
-		HordeConfig.TrimStartInline();
-		HordeConfig.TrimEndInline();
-		HordeConfig.RemoveFromStart(TEXT("("));
-		HordeConfig.RemoveFromEnd(TEXT(")"));
-
-		FString Pool;
-		bool FoundPool = false;
-
-		#if PLATFORM_MAC
-		FoundPool = FParse::Value(*HordeConfig, TEXT("MacPool="), Pool);
-		#endif
-
-		if (!FoundPool)
-		{
-			FoundPool = FParse::Value(*HordeConfig, TEXT("Pool="), Pool);
-		}
-
-		if (FoundPool)
-		{
-			UE_LOG(LogUbaController, Log, TEXT("Found UBA controller Pool: \"%s\""), *Pool);
-			HordeAgentManager->SetPool(Pool);
-		}
-
-		uint32 MaxCores = 0;
-		if (FParse::Value(*HordeConfig, TEXT("MaxCores="), MaxCores))
-		{
-			UE_LOG(LogUbaController, Log, TEXT("Found UBA controller MaxCores: \"%u\""), MaxCores);
-			HordeAgentManager->SetMaxCoreCount(MaxCores);
-		}
+		HordeAgentManager = MakeUnique<FUbaHordeAgentManager>(ControllerModule.GetWorkingDirectory(), GetUbaBinariesPath());
 
-		FString EncryptionName;
-		if (FParse::Value(*HordeConfig, TEXT("Encryption="), EncryptionName))
-		{
-			EHordeTransportEncryption Encryption = EHordeTransportEncryption::None;
-			if (ParseHordeTransportEncryption(Encryption, *EncryptionName))
+		auto AddClientCallback = [](void* UserData, const uba::tchar* Ip, uint16 Port, const uba::tchar* Crypto16Characters)
 			{
-				HordeAgentManager->SetEncryption(Encryption);
-				UE_LOG(LogUbaController, Log, TEXT("Found UBA controller Encryption: \"%s\""), *EncryptionName);
-			}
-			else
-			{
-				UE_LOG(LogUbaController, Log, TEXT("Found UBA controller Encryption with unrecognized identifier: \"%s\""), *EncryptionName);
-			}
-		}
-	}
-
-	FString Host;
-	if (GConfig->GetString(TEXT("UbaController"), TEXT("Host"), Host, GEngineIni))
-	{
-		UE_LOG(LogUbaController, Log, TEXT("Found UBA controller Host: \"%s\""), *Host);
-		HordeAgentManager->SetUbaHost(StringCast<ANSICHAR>(*Host).Get());
-	}
-
-	int32 Port = 0;
-	if (GConfig->GetInt(TEXT("UbaController"), TEXT("Port"), Port, GEngineIni))
-	{
-		UE_LOG(LogUbaController, Log, TEXT("Found UBA controller Port: \"%d\""), Port);
-		HordeAgentManager->SetUbaPort(Port);
-	}
-	
-	if (GConfig->GetBool(TEXT("UbaController"), TEXT("bForceRemote"), bForceRemote, GEngineIni))
-	{
-		UE_LOG(LogUbaController, Log, TEXT("Found UBA controller Force Remote: [%s]"), bForceRemote ? TEXT("True") : TEXT("False"));
+				return NetworkServer_AddClient(reinterpret_cast<uba::NetworkServer*>(UserData), Ip, Port, Crypto16Characters);
+			};
+		HordeAgentManager->SetAddClientCallback(AddClientCallback, UbaServer);
 	}
 
 	UE_LOG(LogUbaController, Display, TEXT("Created UBA storage server: RootDir=%s"), *RootDir);
@@ -502,21 +418,21 @@ uint32 FUbaJobProcessor::Run()
 	{
 		const float ElapsedSeconds = (FPlatformTime::Cycles() - LastTimeSinceHadJobs) * FPlatformTime::GetSecondsPerCycle();
 
-		uint32 queued = 0;
-		uint32 activeLocal = 0;
-		uint32 activeRemote = 0;
-		uint32 finished = 0;
+		uint32 NumQueuedJobs = 0;
+		uint32 NumActiveLocalJobs = 0;
+		uint32 NumActiveRemoteJobs = 0;
+		uint32 NumFinishedJobs = 0;
 
 		FScopeLock lock(&bShouldProcessJobsLock);
 
 		if (UbaScheduler)
 		{
-			Scheduler_GetStats(UbaScheduler, queued, activeLocal, activeRemote, finished);
+			Scheduler_GetStats(UbaScheduler, NumQueuedJobs, NumActiveLocalJobs, NumActiveRemoteJobs, NumFinishedJobs);
 		}
-		uint32 active = activeLocal + activeRemote;
+		const uint32 NumActiveJobs = NumActiveLocalJobs + NumActiveRemoteJobs;
 
 		// We don't want to hog up Horde resources.
-		if (bShouldProcessJobs && ElapsedSeconds > UbaJobProcessorOptions::MaxTimeWithoutTasks && (queued + active) == 0)
+		if (bShouldProcessJobs && ElapsedSeconds > UbaJobProcessorOptions::MaxTimeWithoutTasks && (NumQueuedJobs + NumActiveJobs) == 0)
 		{
 			// If we're optimizing job starting, we only want to shutdown UBA if all the processes have terminated
 			bShouldProcessJobs = false;
@@ -524,7 +440,7 @@ uint32 FUbaJobProcessor::Run()
 		}
 
 		// Check if we have new tasks to process
-		if ((ControllerModule.HasTasksDispatchedOrPending() || (queued + active) != 0))
+		if ((ControllerModule.HasTasksDispatchedOrPending() || (NumQueuedJobs + NumActiveJobs) != 0))
 		{
 			if (!bShouldProcessJobs)
 			{
@@ -539,18 +455,20 @@ uint32 FUbaJobProcessor::Run()
 
 		if (bShouldProcessJobs)
 		{
-			int32 MaxPossible = FPlatformMisc::NumberOfCoresIncludingHyperthreads();
-			int32 LocalCoresToNotUse = 1 + activeRemote / 30; // Use one core per 30 remote ones
-			int32 MaxLocal = FMath::Max(0, MaxPossible - LocalCoresToNotUse);
-			MaxLocal = FMath::Min(MaxLocal, MaxLocalParallelJobs);
-			Scheduler_SetMaxLocalProcessors(UbaScheduler, bForceRemote ? 0 : MaxLocal);
+			const int32 MaxAvailableLocalCores = FPlatformMisc::NumberOfCoresIncludingHyperthreads();
+			const int32 LocalCoresToNotUse = 1 + (int32)NumActiveRemoteJobs / 30; // Use one core per 30 remote ones
+			const int32 MaxLocalCoresToUse = FMath::Clamp(MaxAvailableLocalCores - LocalCoresToNotUse, 0, MaxLocalParallelJobs);
+			const int32 MaxRemoteCoresToRequest = FMath::Max(0, (int32)(NumQueuedJobs + NumActiveJobs) - MaxLocalCoresToUse);
 
-			int32 TargetCoreCount = FMath::Max(0, int32(queued + active) - MaxLocal);
+			Scheduler_SetMaxLocalProcessors(UbaScheduler, MaxLocalCoresToUse);
 
-			HordeAgentManager->SetTargetCoreCount(TargetCoreCount);
+			if (HordeAgentManager)
+			{
+				HordeAgentManager->SetTargetCoreCount(MaxRemoteCoresToRequest);
+			}
 			
 			// TODO: Not sure this is a good idea in a cooking scenario where number of queued processes are going up and down
-			SessionServer_SetMaxRemoteProcessCount(UbaSessionServer, TargetCoreCount);
+			SessionServer_SetMaxRemoteProcessCount(UbaSessionServer, MaxRemoteCoresToRequest);
 
 			UpdateStats();
 		}
@@ -660,7 +578,9 @@ void FUbaJobProcessor::HandleTaskQueueUpdated(const FString& InputFileName)
 	{
 		FDistributedBuildTask* Task = nullptr;
 		if (!ControllerModule.PendingRequestedCompilationTasks.Dequeue(Task) || !Task)
+		{
 			break;
+		}
 		RunTaskWithUba(Task);
 	}
 }
@@ -690,9 +610,12 @@ bool FUbaJobProcessor::PollStats(FDistributedBuildStats& OutStats)
 
 void FUbaJobProcessor::UpdateStats()
 {
-	FScopeLock StatsLockGuard(&StatsLock);
+	if (HordeAgentManager)
+	{
+		FScopeLock StatsLockGuard(&StatsLock);
 
-	// Update maximum
-	Stats.MaxRemoteAgents = FMath::Max(Stats.MaxRemoteAgents, (uint32)HordeAgentManager->GetAgentCount());
-	Stats.MaxActiveAgentCores = FMath::Max(Stats.MaxActiveAgentCores, HordeAgentManager->GetActiveCoreCount());
+		// Update maximum
+		Stats.MaxRemoteAgents = FMath::Max(Stats.MaxRemoteAgents, (uint32)HordeAgentManager->GetAgentCount());
+		Stats.MaxActiveAgentCores = FMath::Max(Stats.MaxActiveAgentCores, HordeAgentManager->GetActiveCoreCount());
+	}
 }
diff --git a/Engine/Plugins/UbaController/Source/Public/UbaJobProcessor.h b/Engine/Plugins/UbaController/Source/Public/UbaJobProcessor.h
--- a/Engine/Plugins/UbaController/Source/Public/UbaJobProcessor.h
+++ b/Engine/Plugins/UbaController/Source/Public/UbaJobProcessor.h
@@ -51,7 +51,7 @@ private:
 	void UpdateStats();
 
 	/** The runnable thread */
-	FRunnableThread* Thread;
+	FRunnableThread* Thread = nullptr;
 	
 	FUbaControllerModule& ControllerModule;
 
@@ -61,14 +61,14 @@ private:
 	TAtomic<bool> bForceStop;
 	
 	/** Used to abort the current processing loop */
-	uint32 LastTimeCheckedForTasks;
+	uint32 LastTimeCheckedForTasks = 0;
 	
 	/** Used to stop the processing loop in the next main loop */
 	FCriticalSection bShouldProcessJobsLock;
-	bool bShouldProcessJobs;
+	bool bShouldProcessJobs = false;
 
 	/** Set to true when the main loop finishes*/
-	bool bIsWorkDone;
+	bool bIsWorkDone = false;
 
 	uba::NetworkServer* UbaServer = nullptr;
 	uba::StorageServer* UbaStorageServer = nullptr;
@@ -83,7 +83,4 @@ private:
 
 	FCriticalSection StatsLock;
 	FDistributedBuildStats Stats;
-
-	/** If true all UBA jobs will be run remotely */
-	bool bForceRemote = false;
 };
\ No newline at end of file
diff --git a/Engine/Source/Developer/Horde/Private/Compute/AgentMessage.cpp b/Engine/Source/Developer/Horde/Private/Compute/AgentMessage.cpp
--- a/Engine/Source/Developer/Horde/Private/Compute/AgentMessage.cpp
+++ b/Engine/Source/Developer/Horde/Private/Compute/AgentMessage.cpp
@@ -63,7 +63,7 @@ void FAgentMessageChannel::UploadFiles(const char* Path, const char* Locator)
 	FlushMessage();
 }
 
-void FAgentMessageChannel::DeleteFiles(const char** Paths, size_t Count)
+void FAgentMessageChannel::DeleteFiles(const char* const * Paths, size_t Count)
 {
 	size_t RequiredSize = 10;
 	for (size_t Idx = 0; Idx < Count; Idx++)
@@ -80,7 +80,7 @@ void FAgentMessageChannel::DeleteFiles(const char** Paths, size_t Count)
 	FlushMessage();
 }
 
-void FAgentMessageChannel::Execute(const char* Exe, const char** Args, size_t NumArgs, const char* WorkingDir, const char** EnvVars, size_t NumEnvVars, EExecuteProcessFlags Flags)
+void FAgentMessageChannel::Execute(const char* Exe, const char* const * Args, size_t NumArgs, const char* WorkingDir, const char* const * EnvVars, size_t NumEnvVars, EExecuteProcessFlags Flags)
 {
 	size_t RequiredSize = 50 + strlen(Exe);
 	for (size_t Idx = 0; Idx < NumArgs; Idx++)
diff --git a/Engine/Source/Developer/Horde/Public/Compute/AgentMessage.h b/Engine/Source/Developer/Horde/Public/Compute/AgentMessage.h
--- a/Engine/Source/Developer/Horde/Public/Compute/AgentMessage.h
+++ b/Engine/Source/Developer/Horde/Public/Compute/AgentMessage.h
@@ -117,10 +117,10 @@ public:
 	HORDE_API void UploadFiles(const char* Path, const char* Locator);
 
 	// Deletes files matching a set of wildcards
-	HORDE_API void DeleteFiles(const char** Paths, size_t Count);
+	HORDE_API void DeleteFiles(const char* const * Paths, size_t Count);
 
 	// Executes a process on the remote machine
-	HORDE_API void Execute(const char* Exe, const char** Args, size_t NumArgs, const char* WorkingDir, const char** EnvVars, size_t NumEnvVars, EExecuteProcessFlags Flags);
+	HORDE_API void Execute(const char* Exe, const char* const * Args, size_t NumArgs, const char* WorkingDir, const char* const * EnvVars, size_t NumEnvVars, EExecuteProcessFlags Flags);
 
 	// Writes a blob requested by the remote
 	HORDE_API void Blob(const unsigned char* Data, size_t Length);
diff --git a/Engine/Source/Developer/UbaCoordinatorHorde/Private/UbaHordeAgent.cpp b/Engine/Source/Developer/UbaCoordinatorHorde/Private/UbaHordeAgent.cpp
--- a/Engine/Source/Developer/UbaCoordinatorHorde/Private/UbaHordeAgent.cpp
+++ b/Engine/Source/Developer/UbaCoordinatorHorde/Private/UbaHordeAgent.cpp
@@ -1,6 +1,7 @@
 // Copyright Epic Games, Inc. All Rights Reserved.
 
 #include "UbaHordeAgent.h"
+#include "UbaHordeConfig.h"
 
 #include "Misc/Paths.h"
 #include "HAL/FileManager.h"
@@ -28,7 +29,7 @@ FUbaHordeAgent::FUbaHordeAgent(const FHordeRemoteMachineInfo& InMachineInfo)
 	ComputeTransport->Send(MachineInfo.Nonce, sizeof(MachineInfo.Nonce));
 
 	// If encryption protocol is requested, wrap initial transport layer around an encryption layer *after* initial nonce was sent unencrypted
-	if (InMachineInfo.Encryption == EHordeTransportEncryption::AES)
+	if (InMachineInfo.Encryption == EUbaHordeEncryption::AES)
 	{
 		ComputeTransport = MakeUnique<FUbaHordeComputeTransportAES>(MachineInfo, MoveTemp(ComputeTransport));
 		if (!ComputeTransport->IsValid())
@@ -181,7 +182,7 @@ bool FUbaHordeAgent::UploadBinaries(const FString& BundleDirectory, const char*
 	return bSuccess;
 }
 
-void FUbaHordeAgent::Execute(const char* Exe, const char** Args, size_t NumArgs, const char* WorkingDir, const char** EnvVars, size_t NumEnvVars, bool bUseWine)
+void FUbaHordeAgent::Execute(const char* Exe, const char* const * Args, size_t NumArgs, const char* WorkingDir, const char* const * EnvVars, size_t NumEnvVars, bool bUseWine)
 {
 	ChildChannel->Execute(Exe, Args, NumArgs, WorkingDir, EnvVars, NumEnvVars, bUseWine ? EExecuteProcessFlags::UseWine : EExecuteProcessFlags::None);
 }
diff --git a/Engine/Source/Developer/UbaCoordinatorHorde/Private/UbaHordeAgent.h b/Engine/Source/Developer/UbaCoordinatorHorde/Private/UbaHordeAgent.h
--- a/Engine/Source/Developer/UbaCoordinatorHorde/Private/UbaHordeAgent.h
+++ b/Engine/Source/Developer/UbaCoordinatorHorde/Private/UbaHordeAgent.h
@@ -24,7 +24,7 @@ public:
 	 * e.g. Locator = abcdef1234567890abcdef1234567890abcdef12@34567890abcdef1234567890#123 */
 	bool UploadBinaries(const FString& BundleDirectory, const char* BundleLocator);
 
-	void Execute(const char* Exe, const char** Args, size_t NumArgs, const char* WorkingDir, const char** EnvVars, size_t NumEnvVars, bool bUseWine = false);
+	void Execute(const char* Exe, const char* const * Args, size_t NumArgs, const char* WorkingDir, const char* const * EnvVars, size_t NumEnvVars, bool bUseWine = false);
 
 	void CloseConnection(); 
 
diff --git a/Engine/Source/Developer/UbaCoordinatorHorde/Private/UbaHordeAgentManager.cpp b/Engine/Source/Developer/UbaCoordinatorHorde/Private/UbaHordeAgentManager.cpp
--- a/Engine/Source/Developer/UbaCoordinatorHorde/Private/UbaHordeAgentManager.cpp
+++ b/Engine/Source/Developer/UbaCoordinatorHorde/Private/UbaHordeAgentManager.cpp
@@ -17,6 +17,7 @@
 #include "Storage/BlobWriter.h"
 #include "UbaBase.h"
 #include "UbaHordeAgent.h"
+#include "UbaDefaultConstants.h"
 #include <filesystem>
 #include <fstream>
 
@@ -53,7 +54,7 @@ FUbaHordeAgentManager::~FUbaHordeAgentManager()
 
 void FUbaHordeAgentManager::SetTargetCoreCount(uint32 Count)
 {
-	TargetCoreCount = FMath::Min(MaxCores, Count);
+	TargetCoreCount = FMath::Min((uint32)FUbaHordeConfig::Get().HordeMaxCores, Count);
 
 	while (EstimatedCoreCount < TargetCoreCount)
 	{
@@ -79,10 +80,10 @@ void FUbaHordeAgentManager::SetTargetCoreCount(uint32 Count)
 	}
 }
 
-void FUbaHordeAgentManager::SetAddClientCallback(AddClientCallback* callback, void* userData)
+void FUbaHordeAgentManager::SetAddClientCallback(FAddClientCallback* Callback, void* UserData)
 {
-	m_callback = callback;
-	m_userData = userData;
+	AddClientCallback = Callback;
+	AddClientUserData = UserData;
 }
 
 int32 FUbaHordeAgentManager::GetAgentCount() const
@@ -162,15 +163,52 @@ void FUbaHordeAgentManager::RequestAgent()
 {
 	EstimatedCoreCount += 32; // We estimate a typical agent to have 32 cores
 
+	const int32 NextAgentIndex = Agents.Num();
+
 	FScopeLock AgentsScopeLock(&AgentsLock);
 	FHordeAgentWrapper& Wrapper = *Agents.Emplace_GetRef(MakeUnique<FHordeAgentWrapper>());
 
 	Wrapper.ShouldExit = FGenericPlatformProcess::GetSynchEventFromPool(true);
-	Wrapper.Thread = FThread(TEXT("HordeAgent"), [this, WrapperPtr = &Wrapper]() { ThreadAgent(*WrapperPtr); });
+
+	Wrapper.Thread = FThread(*FString::Printf(TEXT("HordeAgent %d"), NextAgentIndex), [this, &Wrapper]() { ThreadAgent(Wrapper); });
 }
 
+// Helper class to manage raw and auto-released strings for UBA agent command line arguments.
+class FUbaAgentArguments
+{
+	TArray<const char*> RawArguments;
+	TArray<FAnsiString> ArgumentStrings;
+public:
+	FUbaAgentArguments(int32 InitialArgumentCapapcity = 16)
+	{
+		RawArguments.Reserve(InitialArgumentCapapcity);
+	}
+	void AddRaw(const char* InArgument)
+	{
+		RawArguments.Add(InArgument);
+	}
+	void Add(const FAnsiString& InArgument)
+	{
+		ArgumentStrings.Add(InArgument);
+	}
+	void Finalize()
+	{
+		RawArguments.Reserve(RawArguments.Num() + ArgumentStrings.Num());
+		for (const FAnsiString& Argument : ArgumentStrings)
+		{
+			RawArguments.Add(Argument.GetCharArray().GetData());
+		}
+	}
+	const TArray<const char*>& Get() const
+	{
+		return RawArguments;
+	}
+};
+
 void FUbaHordeAgentManager::ThreadAgent(FHordeAgentWrapper& Wrapper)
 {
+	const FUbaHordeConfig& Config = FUbaHordeConfig::Get();
+
 	FEvent& ShouldExit = *Wrapper.ShouldExit;
 	TUniquePtr<FUbaHordeAgent> Agent;
 	bool bSuccess = false;
@@ -181,14 +219,16 @@ void FUbaHordeAgentManager::ThreadAgent(FHordeAgentWrapper& Wrapper)
 		{
 			Agent->CloseConnection();
 		}
-
 		ShouldExit.Trigger();
 	};
 
 	int MachineCoreCount = 0;
 
 	// If no host is specified, we need to start the agent in listen mode
-	const bool bUseListen = UbaHost.IsEmpty();
+	const bool bUseListen = Config.HordeHost.IsEmpty();
+
+	FString UbaAgentIp;
+	FHordeRemoteMachineInfo::FPortInfo UbaPort, UbaProxyPort;
 
 	{
 		ON_SCOPE_EXIT{ EstimatedCoreCount -= 32; };
@@ -271,10 +311,31 @@ void FUbaHordeAgentManager::ThreadAgent(FHordeAgentWrapper& Wrapper)
 			}
 		}
 
-		TSharedPtr<FUbaHordeMetaClient::HordeMachinePromise, ESPMode::ThreadSafe> Promise = HordeMetaClient->RequestMachine(Pool, nullptr, Encryption);
+		// Initialize horde request with connection mode preference and encryption.
+		// Pool ID is superseded by cluster ID, but it still has to be provided if default cluster and direct mode is used.
+		const FString HordeRequestJsonBody = HordeMetaClient->BuildHordeRequestJsonBody(Config.HordeConnectionMode == EUbaHordeConnectionMode::Direct ? Config.GetHordePool() : TEXT(""), Config.HordeConnectionMode, Config.HordeEncryption);
+
+		// Resolve cluster ID through Horde if "_auto" was specified
+		FHordeClusterInfo ClusterInfo;
+		if (Config.HordeCluster == FUbaHordeConfig::ClusterAuto)
+		{
+			TSharedPtr<FUbaHordeMetaClient::HordeClusterPromise, ESPMode::ThreadSafe> Promise = HordeMetaClient->RequestClusterId(HordeRequestJsonBody);
+			if (!Promise)
+			{
+				UE_LOG(LogUbaHorde, Log, TEXT("Failed to resolve Horde cluster ID with HTTP/Json request: %s"), *HordeRequestJsonBody);
+				return;
+			}
+			TFuture<TTuple<FHttpResponsePtr, FHordeClusterInfo>> Future = Promise->GetFuture();
+			Future.Wait();
+			ClusterInfo = Future.Get().Value;
+			UE_CLOG(!ClusterInfo.ClusterId.IsEmpty(), LogUbaHorde, Verbose, TEXT("Received resolved cluster ID from Horde: %s"), *ClusterInfo.ClusterId);
+		}
+
+		// Request Horde machine for UBA agent
+		TSharedPtr<FUbaHordeMetaClient::HordeMachinePromise, ESPMode::ThreadSafe> Promise = HordeMetaClient->RequestMachine(HordeRequestJsonBody, *ClusterInfo.ClusterId);
 		if (!Promise)
 		{
-			//UE_LOG(LogUbaHorde, Error, TEXT("Failed to create Horde bundle for UbaAgent executable: %s"), *UbaAgentFilePath);
+			UE_LOG(LogUbaHorde, Log, TEXT("Failed to request Horde machine from cluster: %s"), *ClusterInfo.ClusterId);
 			return;
 		}
 		TFuture<TTuple<FHttpResponsePtr, FHordeRemoteMachineInfo>> Future = Promise->GetFuture();
@@ -282,12 +343,9 @@ void FUbaHordeAgentManager::ThreadAgent(FHordeAgentWrapper& Wrapper)
 		FHordeRemoteMachineInfo MachineInfo = Future.Get().Value;
 
 		// If the machine couldn't be assigned, just ignore this agent slot
-		if (MachineInfo.Ip == TEXT(""))
+		if (MachineInfo.GetConnectionAddress().IsEmpty())
 		{
-			if (!LastRequestFailTime)
-			{
-				UE_LOG(LogUbaHorde, Verbose, TEXT("No resources available in Horde. Will keep retrying until %u cores are used (Currently have %u)"), TargetCoreCount.Load(), ActiveCoreCount.Load());
-			}
+			UE_CLOG(!LastRequestFailTime, LogUbaHorde, Verbose, TEXT("No resources available in Horde. Will keep retrying until %u cores are used (Currently have %u)"), TargetCoreCount.Load(), ActiveCoreCount.Load());
 			LastRequestFailTime = FPlatformTime::Cycles64();
 			return;
 		}
@@ -336,20 +394,47 @@ void FUbaHordeAgentManager::ThreadAgent(FHordeAgentWrapper& Wrapper)
 			}
 		}
 
+		// Get IP address and port for UbaAgent
+		UbaAgentIp = Agent->GetMachineInfo().GetConnectionAddress();
+
+		UbaPort = Agent->GetMachineInfo().Ports[TEXT("UbaPort")];
+		UbaProxyPort = Agent->GetMachineInfo().Ports[TEXT("UbaProxyPort")];
+
+		const uint32 UbaHostPort = uba::DefaultPort;
+
 		// Start the UBA Agent that will connect to us, requesting for work
+		FUbaAgentArguments UbaAgentArgs;
 
-		const FAnsiString AgentConnectionArg = bUseListen ? FAnsiString::Printf("-listen=%u", UbaPort) : FAnsiString::Printf("-Host=%s:%u", *UbaHost, UbaPort);
+		if (bUseListen)
+		{
+			UbaAgentArgs.Add(FAnsiString::Printf("-Listen=%u", UbaPort.AgentPort));
+			UbaAgentArgs.Add(TEXT("-ListenTimeout=10"));
+		}
+		else
+		{
+			UbaAgentArgs.Add(FAnsiString::Printf("-Host=%s:%u", TCHAR_TO_ANSI(*Config.HordeHost), UbaHostPort));
+		}
 
-		const char* UbaAgentArgs[] =
+		if (!Config.UbaSentryUrl.IsEmpty())
 		{
-			*AgentConnectionArg,
-			"-nopoll",				// -nopoll recommended when running on remote Horde agents to make sure they exit after completion. Otherwise, it keeps running.
-			"-listenTimeout=5",		// Agent will wait 5 seconds for this thread to connect (Server_AddClient does the connect)
-			"-quiet",				// Skip all the agent logging that would be sent over to here
-			"-maxidle=15",			// After 15 seconds of idling agent will automatically disconnect
-			"-Dir=%UE_HORDE_SHARED_DIR%\\Uba",
-			"-Eventfile=%UE_HORDE_TERMINATION_SIGNAL_FILE%",
-		};
+			UbaAgentArgs.Add(FAnsiString::Printf("-Sentry=%s", TCHAR_TO_ANSI(*Config.UbaSentryUrl)));
+		}
+
+		UbaAgentArgs.Add(FAnsiString::Printf("-ProxyPort=%u", UbaProxyPort.AgentPort));
+
+		// -nopoll recommended when running on remote Horde agents to make sure they exit after completion. Otherwise, it keeps running.
+		UbaAgentArgs.AddRaw("-NoPoll");
+
+		// Skip all the agent logging that would be sent over to here
+		UbaAgentArgs.AddRaw("-Quiet");
+
+		// After 15 seconds of idling agent will automatically disconnect
+		UbaAgentArgs.AddRaw("-MaxIdle=15");
+
+		UbaAgentArgs.AddRaw("-Dir=%UE_HORDE_SHARED_DIR%\\Uba");
+		UbaAgentArgs.AddRaw("-Eventfile=%UE_HORDE_TERMINATION_SIGNAL_FILE%");
+
+		UbaAgentArgs.Finalize();
 
 		// If the machine does not run Windows, enable the compatibility layer Wine to run UbaAgent.exe on POSIX systems
 		#if PLATFORM_WINDOWS
@@ -364,16 +449,16 @@ void FUbaHordeAgentManager::ThreadAgent(FHordeAgentWrapper& Wrapper)
 			return;
 		}
 
-		Agent->Execute(AppName, UbaAgentArgs, UE_ARRAY_COUNT(UbaAgentArgs), nullptr, nullptr, 0, bUseWine);
+		Agent->Execute(AppName, UbaAgentArgs.Get().GetData(), (size_t)UbaAgentArgs.Get().Num(), nullptr, nullptr, 0, bUseWine);
 
 		// Log remote execution
 		FString UbaAgentCmdArgs = ANSI_TO_TCHAR(AppName);
-		for (const char* Arg : UbaAgentArgs)
+		for (const char* Arg : UbaAgentArgs.Get())
 		{
 			UbaAgentCmdArgs += TEXT(" ");
 			UbaAgentCmdArgs += ANSI_TO_TCHAR(Arg);
 		}
-		UE_LOG(LogUbaHorde, Log, TEXT("Remote execution on Horde machine [%s:%u]: %s"), *Agent->GetMachineInfo().Ip, UbaPort, *UbaAgentCmdArgs);
+		UE_LOG(LogUbaHorde, Log, TEXT("Remote execution on Horde machine [%s:%u]: %s"), *UbaAgentIp, UbaPort.Port, *UbaAgentCmdArgs);
 
 		MachineCoreCount = MachineInfo.LogicalCores;
 		EstimatedCoreCount += MachineCoreCount;
@@ -395,15 +480,32 @@ void FUbaHordeAgentManager::ThreadAgent(FHordeAgentWrapper& Wrapper)
 
 		if (callCounter++ == 2)
 		{
-			// Add this machine as client to the remote agent
-			const FString& IpAddress = Agent->GetMachineInfo().Ip;
-			const bool bAddClientSuccess = 	m_callback(m_userData, StringCast<uba::tchar>(*IpAddress).Get(), static_cast<uint16>(UbaPort));
+			FString CryptoNonce16;
+			if (Agent->GetMachineInfo().Encryption != EUbaHordeEncryption::None)
+			{
+				// Generate random crypto nonce for UbaAgent
+				CryptoNonce16.Reserve(16);
+				for (int32 CharacterIndex = 0; CharacterIndex < 16; ++CharacterIndex)
+				{
+					CryptoNonce16.AppendChar(TEXT("0123456789abcdef")[FMath::Rand() % 0xF]);
+				}
+			}
+
+			// Add this machine as client to the remote agent.
+			// Use JSON field "port" when adding clients while "agentPort" is used for UbaAgent to listen to on the server side.
+			const bool bAddClientSuccess = AddClientCallback(
+				AddClientUserData,
+				StringCast<uba::tchar>(*UbaAgentIp).Get(),
+				static_cast<uint16>(UbaPort.Port),
+				Agent->GetMachineInfo().Encryption != EUbaHordeEncryption::None ? StringCast<uba::tchar>(*CryptoNonce16).Get() : nullptr
+			);
 
 			if (!bAddClientSuccess)
 			{
-				UE_LOG(LogUbaHorde, Display, TEXT("Server_AddClient(%s:%u) failed"), *IpAddress, UbaPort);
+				UE_LOG(LogUbaHorde, Display, TEXT("Server_AddClient(%s:%u) failed"), *UbaAgentIp, UbaPort.Port);
 				return;
 			}
 		}
 	}
 }
+
diff --git a/Engine/Source/Developer/UbaCoordinatorHorde/Private/UbaHordeComputeTransport.cpp b/Engine/Source/Developer/UbaCoordinatorHorde/Private/UbaHordeComputeTransport.cpp
--- a/Engine/Source/Developer/UbaCoordinatorHorde/Private/UbaHordeComputeTransport.cpp
+++ b/Engine/Source/Developer/UbaCoordinatorHorde/Private/UbaHordeComputeTransport.cpp
@@ -14,12 +14,12 @@ FUbaHordeComputeTransport::FUbaHordeComputeTransport(const FHordeRemoteMachineIn
 	Socket = SocketSubsystem->CreateSocket(TEXT("Stream"), TEXT(""));
 	TSharedPtr<FInternetAddr> InternetAddr;
 
-	TSharedPtr<FInternetAddr> Address = SocketSubsystem->GetAddressFromString(MachineInfo.Ip);
-	Address->SetPort(MachineInfo.Port);
+	TSharedPtr<FInternetAddr> Address = SocketSubsystem->GetAddressFromString(MachineInfo.GetConnectionAddress());
+	Address->SetPort(MachineInfo.GetConnectionPort().Port);
 
 	if (!Socket->Connect(*Address))
 	{
-		UE_LOG(LogUbaHorde, Display, TEXT("Failed to connect to Horde Agent [%s:%d]"), *MachineInfo.Ip, (int32)MachineInfo.Port);
+		UE_LOG(LogUbaHorde, Display, TEXT("Failed to connect to Horde compute service [%s]"), *Address->ToString(true));
 
 		SocketSubsystem->DestroySocket(Socket);
 
@@ -43,12 +43,12 @@ size_t FUbaHordeComputeTransport::Send(const void* Data, size_t Size)
 	if (!Socket->Send((const uint8*)Data, Size, NumBytesSent))
 	{
 		// We can't log this since the other side could have disconnected us (causing recv to fail). This happens often in the horde setup
-		//UE_LOG(LogUbaHorde, Error, TEXT("Failed to send data to the Horde Agent: %llu %s"), (uint64)Size, Size == 1 ? TEXT("byte") : TEXT("bytes"));
+		//UE_LOG(LogUbaHorde, Error, TEXT("Failed to send data to the Horde compute service: %llu %s"), (uint64)Size, Size == 1 ? TEXT("byte") : TEXT("bytes"));
 		bHasErrors = true;
 		return 0;
 	}
 
-	UE_LOG(LogUbaHorde, Verbose, TEXT("Sent message to horde agent: %d %s"), NumBytesSent, NumBytesSent == 1 ? TEXT("byte") : TEXT("bytes"));
+	UE_LOG(LogUbaHorde, Verbose, TEXT("Sent message to Horde compute service: %d %s"), NumBytesSent, NumBytesSent == 1 ? TEXT("byte") : TEXT("bytes"));
 
 	return NumBytesSent;
 }
@@ -61,11 +61,11 @@ size_t FUbaHordeComputeTransport::Recv(void* Data, size_t Size)
 	{
 		// We can't log this since the other side could have disconnected us (causing recv to fail). This happens often in the horde setup
 		//if (!bIsClosed)
-		//	UE_LOG(LogUbaHorde, Error, TEXT("Failed to receive data from the Horde Agent: %llu %s"), (uint64)Size, Size == 1 ? TEXT("byte") : TEXT("bytes"));
+		//	UE_LOG(LogUbaHorde, Error, TEXT("Failed to receive data from the Horde compute service: %llu %s"), (uint64)Size, Size == 1 ? TEXT("byte") : TEXT("bytes"));
 		return 0;
 	}
 
-	UE_LOG(LogUbaHorde, Verbose, TEXT("Received message from horde agent: %d %s"), NumBytesRead, NumBytesRead == 1 ? TEXT("byte") : TEXT("bytes"));
+	UE_LOG(LogUbaHorde, Verbose, TEXT("Received message from Horde compute service: %d %s"), NumBytesRead, NumBytesRead == 1 ? TEXT("byte") : TEXT("bytes"));
 
 	return NumBytesRead;
 }
diff --git a/Engine/Source/Developer/UbaCoordinatorHorde/Private/UbaHordeConfig.cpp b/Engine/Source/Developer/UbaCoordinatorHorde/Private/UbaHordeConfig.cpp
new file mode 100644
--- /dev/null
+++ b/Engine/Source/Developer/UbaCoordinatorHorde/Private/UbaHordeConfig.cpp
@@ -0,0 +1,433 @@
+// Copyright Epic Games, Inc. All Rights Reserved.
+
+#include "UbaHordeConfig.h"
+#include "UbaHordeMetaClient.h"
+
+#include "HAL/PlatformProcess.h"
+#include "Misc/App.h"
+#include "Misc/ConfigCacheIni.h"
+#include "Misc/Paths.h"
+#include "XmlParser.h"
+
+
+const TCHAR* FUbaHordeConfig::ClusterDefault = TEXT("default");
+const TCHAR* FUbaHordeConfig::ClusterAuto = TEXT("_auto");
+
+UBACOORDINATORHORDE_API const FString& FUbaHordeConfig::GetHordePool() const
+{
+	if (!OverrideHordePool.IsEmpty())
+	{
+		return OverrideHordePool;
+	}
+#if PLATFORM_WINDOWS
+	if (!WindowsHordePool.IsEmpty())
+	{
+		return WindowsHordePool;
+	}
+#endif
+#if PLATFORM_MAC
+	if (!MacHordePool.IsEmpty())
+	{
+		return MacHordePool;
+	}
+#endif
+#if PLATFORM_LINUX
+	if (!LinuxHordePool.IsEmpty())
+	{
+		return LinuxHordePool;
+	}
+#endif
+	return DefaultHordePool;
+}
+
+static bool ParseUbaHordeConfigFromXmlFile(FUbaHordeConfig& OutConfig, const TCHAR* InBuildConfigFilename)
+{
+	FXmlFile XmlFile;
+	if (!XmlFile.LoadFile(InBuildConfigFilename))
+	{
+		return false;
+	}
+
+	static const FString ConfigurationTag(TEXT("Configuration"));
+	static const FString BuildConfigurationTag(TEXT("BuildConfiguration"));
+	static const FString HordeTag(TEXT("Horde"));
+	static const FString UbaDisabledTag(TEXT("bDisableHorde"));
+	static const FString ServerTag(TEXT("Server"));
+	static const FString ConnectionModeTag(TEXT("ConnectionMode"));
+	static const FString EncryptionTag(TEXT("Encryption"));
+	static const FString ClusterTag(TEXT("Cluster"));
+	static const FString DefaultPoolTag(TEXT("Pool"));
+	static const FString MacPoolTag(TEXT("MacPool"));
+	static const FString UBASentryUrlTag(TEXT("UBASentryUrl"));
+
+	const FXmlNode* ConfigurationNode = XmlFile.GetRootNode();
+	if (ConfigurationNode == nullptr || ConfigurationNode->GetTag() != ConfigurationTag)
+	{
+		return false;
+	}
+
+	if (const FXmlNode* HordeNode = ConfigurationNode->FindChildNode(HordeTag))
+	{
+		if (const FXmlNode* UbaDisabledNode = HordeNode->FindChildNode(UbaDisabledTag))
+		{
+			if (UbaDisabledNode->GetContent() == TEXT("true"))
+			{
+				OutConfig.HordeEnabled = EUbaHordeEnabledMode::Disabled;
+			}
+		}
+		if (const FXmlNode* DefaultPoolNode = HordeNode->FindChildNode(DefaultPoolTag))
+		{
+			OutConfig.DefaultHordePool = DefaultPoolNode->GetContent();
+		}
+		if (const FXmlNode* MacPoolNode = HordeNode->FindChildNode(MacPoolTag))
+		{
+			OutConfig.MacHordePool = MacPoolNode->GetContent();
+		}
+		if (const FXmlNode* ServerNode = HordeNode->FindChildNode(ServerTag))
+		{
+			OutConfig.HordeServer = ServerNode->GetContent();
+		}
+		if (const FXmlNode* ConnectionModeNode = HordeNode->FindChildNode(ConnectionModeTag))
+		{
+			LexFromString(OutConfig.HordeConnectionMode, *ConnectionModeNode->GetContent());
+		}
+		if (const FXmlNode* EncryptionNode = HordeNode->FindChildNode(EncryptionTag))
+		{
+			LexFromString(OutConfig.HordeEncryption, *EncryptionNode->GetContent());
+		}
+		if (const FXmlNode* ClusterNode = HordeNode->FindChildNode(ClusterTag))
+		{
+			OutConfig.HordeCluster = ClusterNode->GetContent();
+		}
+		if (const FXmlNode* UBASentryUrlNode = HordeNode->FindChildNode(UBASentryUrlTag))
+		{
+			OutConfig.UbaSentryUrl = UBASentryUrlNode->GetContent();
+		}
+	}
+
+	return true;
+}
+
+static void ParseUbaHordeConfigFromIniFile(FUbaHordeConfig& OutConfig)
+{
+	GConfig->GetString(TEXT("Horde"), TEXT("ServerUrl"), OutConfig.HordeServer, GEngineIni);
+	GConfig->GetString(TEXT("Horde"), TEXT("UbaToken"), OutConfig.HordeToken, GEngineIni);
+	GConfig->GetString(TEXT("Horde"), TEXT("UbaPool"), OutConfig.DefaultHordePool, GEngineIni);
+	GConfig->GetString(TEXT("Horde"), TEXT("UbaLinuxPool"), OutConfig.LinuxHordePool, GEngineIni);
+	GConfig->GetString(TEXT("Horde"), TEXT("UbaMacPool"), OutConfig.MacHordePool, GEngineIni);
+	GConfig->GetString(TEXT("Horde"), TEXT("UbaWindowsPool"), OutConfig.WindowsHordePool, GEngineIni);
+	GConfig->GetString(TEXT("Horde"), TEXT("UbaCluster"), OutConfig.HordeCluster, GEngineIni);
+	GConfig->GetString(TEXT("Horde"), TEXT("UbaLocalHost"), OutConfig.HordeHost, GEngineIni);
+	GConfig->GetInt(TEXT("Horde"), TEXT("MaxCores"), OutConfig.HordeMaxCores, GEngineIni);
+	GConfig->GetBool(TEXT("Horde"), TEXT("UbaAllowWine"), OutConfig.bHordeAllowWine, GEngineIni);
+
+	FString UbaConnectionModeString;
+	if (GConfig->GetString(TEXT("Horde"), TEXT("UbaConnectionMode"), UbaConnectionModeString, GEngineIni))
+	{
+		LexFromString(OutConfig.HordeConnectionMode, *UbaConnectionModeString);
+	}
+
+	FString UbaEncryptionString;
+	if (GConfig->GetString(TEXT("Horde"), TEXT("UbaEncryption"), UbaEncryptionString, GEngineIni))
+	{
+		LexFromString(OutConfig.HordeEncryption, *UbaEncryptionString);
+	}
+
+	GConfig->GetString(TEXT("Horde"), TEXT("UbaSentryUrl"), OutConfig.UbaSentryUrl, GEngineIni);
+
+	FString UbaEnabledString;
+	if (GConfig->GetString(TEXT("UbaController"), TEXT("UbaEnabled"), UbaEnabledString, GEngineIni))
+	{
+		LexFromString(OutConfig.HordeEnabled, *UbaEnabledString);
+	}
+}
+
+static bool ParseCmdLineBoolean(const TCHAR* CmdLine, const TCHAR* Parameter, bool& bOutValue)
+{
+	if (FParse::Param(CmdLine, Parameter))
+	{
+		FString ArgumentStringValue;
+		if (FParse::Value(CmdLine, *FString::Printf(TEXT("%s="), Parameter), ArgumentStringValue))
+		{
+			bOutValue = ArgumentStringValue != TEXT("false") && ArgumentStringValue != TEXT("0");
+		}
+		else
+		{
+			bOutValue = true;
+		}
+		return true;
+	}
+	return false;
+}
+
+template <typename T>
+static bool ParseCmdLineEnum(const TCHAR* CmdLine, const TCHAR* Parameter, T& OutValue)
+{
+	FString ArgumentStringValue;
+	if (FParse::Value(CmdLine, Parameter, ArgumentStringValue))
+	{
+		LexFromString(OutValue, *ArgumentStringValue);
+		return true;
+	}
+	return false;
+}
+
+static void ParseUbaHordeConfigFromCommandline(FUbaHordeConfig& OutConfig, const TCHAR* CmdLine)
+{
+	FParse::Value(CmdLine, TEXT("UBAHorde="), OutConfig.HordeServer);
+	FParse::Value(CmdLine, TEXT("UBAHordeToken="), OutConfig.HordeToken);
+	FParse::Value(CmdLine, TEXT("UBAHordePool="), OutConfig.OverrideHordePool);
+	ParseCmdLineBoolean(CmdLine, TEXT("UBAHordeAllowWine"), OutConfig.bHordeAllowWine);
+	FParse::Value(CmdLine, TEXT("UBAHordeMaxCores="), OutConfig.HordeMaxCores);
+	FParse::Value(CmdLine, TEXT("UBAHordeMaxLocalCores="), OutConfig.HordeMaxLocalCores);
+	FParse::Value(CmdLine, TEXT("UBAHordeHost="), OutConfig.HordeHost);
+	FParse::Value(CmdLine, TEXT("UBAHordeCluster="), OutConfig.HordeCluster);
+	ParseCmdLineEnum(CmdLine, TEXT("UBAHordeConnectionMode="), OutConfig.HordeConnectionMode);
+	ParseCmdLineEnum(CmdLine, TEXT("UBAHordeEncryption="), OutConfig.HordeEncryption);
+	FParse::Value(CmdLine, TEXT("UBASentryUrl="), OutConfig.UbaSentryUrl);
+	if (FParse::Param(CmdLine, TEXT("UBADisableHorde")))
+	{
+		OutConfig.HordeEnabled = EUbaHordeEnabledMode::Disabled;
+	}
+	else if (FParse::Param(CmdLine, TEXT("UBAEnableHorde")))
+	{
+		OutConfig.HordeEnabled = EUbaHordeEnabledMode::Enabled;
+	}
+	if (FParse::Param(CmdLine, TEXT("ExclusiveRemoteShaderCompiling")))
+	{
+		OutConfig.HordeMaxLocalCores = 0; // -ExclusiveRemoteShaderCompiling overrides -UBAHordeMaxLocalCores
+	}
+}
+
+// Matches same list of XML files to parse the configuration from as in XmlConfig.cs
+static TArray<FString> GetBuildConfigFilePaths()
+{
+	//	"C:\ProgramData\Unreal Engine\UnrealBuildTool\BuildConfiguration.xml"
+	const TCHAR* UEFolderName = TEXT("Unreal Engine");
+	const TCHAR* UBTFolderName = TEXT("UnrealBuildTool");
+	const TCHAR* BuildConfigFilename = TEXT("BuildConfiguration.xml");
+
+	TArray<FString> OutPaths;
+
+	// Skip all the config files under the Engine folder if it's an installed build
+	if (!FApp::IsEngineInstalled())
+	{
+		// Check for the engine config file under /Engine/Programs/NotForLicensees/UnrealBuildTool
+		OutPaths.Add(FPaths::Combine(FPaths::EngineDir(), TEXT("Restricted"), TEXT("NotForLicensees"), TEXT("Programs"), UBTFolderName, BuildConfigFilename));
+
+		// Check for the engine user config file under /Engine/Saved/UnrealBuildTool
+		OutPaths.Add(FPaths::Combine(FPaths::EngineDir(), TEXT("Saved"), UBTFolderName, BuildConfigFilename));
+	}
+
+	// Check for the global config file under ProgramData/Unreal Engine/UnrealBuildTool
+	const FString ProgramDataFolder = FPlatformProcess::GetApplicationSettingsDir(FPlatformProcess::ApplicationSettingsContext{ FPlatformProcess::ApplicationSettingsContext::Context::ApplicationSpecific, /*bIsEpic:*/ false });
+	if (!ProgramDataFolder.IsEmpty())
+	{
+		OutPaths.Add(FPaths::Combine(ProgramDataFolder, UEFolderName, UBTFolderName, BuildConfigFilename));
+	}
+
+	// Check for the global config file under AppData/Unreal Engine/UnrealBuildTool (Roaming)
+	const FString AppDataFolder = FPlatformProcess::GetApplicationSettingsDir(FPlatformProcess::ApplicationSettingsContext{ FPlatformProcess::ApplicationSettingsContext::Context::RoamingUser, /*bIsEpic:*/ false });
+	if (!AppDataFolder.IsEmpty())
+	{
+		OutPaths.Add(FPaths::Combine(AppDataFolder, UEFolderName, UBTFolderName, BuildConfigFilename));
+	}
+
+	// Check for the global config file under LocalAppData/Unreal Engine/UnrealBuildTool
+	if (const TCHAR* LocalAppDataFolder = FPlatformProcess::UserSettingsDir())
+	{
+		OutPaths.Add(FPaths::Combine(LocalAppDataFolder, UEFolderName, UBTFolderName, BuildConfigFilename));
+	}
+
+	// Check for the global config file under My Documents/Unreal Engine/UnrealBuildTool
+	if (const TCHAR* PersonalFolder = FPlatformProcess::UserDir())
+	{
+		OutPaths.Add(FPaths::Combine(PersonalFolder, UEFolderName, UBTFolderName, BuildConfigFilename));
+	}
+
+	return OutPaths;
+}
+
+/**
+ * Reads the UBA / Horde configuration in the following order :
+ *  1. From a list of BuildConfiguration.xml files (in application settings and user folders)
+ *  2. From BaseEngine.ini configuration files under [Horde] section
+ *  3. From command line argument (e.g. -UBAHorde=...)
+ */
+static void ReadUbaHordeConfig(FUbaHordeConfig& OutConfig)
+{
+	// Read configurations from BuildConfiguration.xml files
+	for (const FString& ConfigFilePath : GetBuildConfigFilePaths())
+	{
+		if (ParseUbaHordeConfigFromXmlFile(OutConfig, *ConfigFilePath))
+		{
+			UE_LOG(LogUbaHorde, Verbose, TEXT("Read UBA/Horde settings from %s"), *ConfigFilePath);
+		}
+	}
+
+	// Override configuration entries by their configuration INI files
+	ParseUbaHordeConfigFromIniFile(OutConfig);
+
+	// Override configuration entries by their command line arguments
+	ParseUbaHordeConfigFromCommandline(OutConfig, FCommandLine::Get());
+
+	// If connection mode is "realy"-mode, encryption is implied and we use AES since UbaController does not support SSL/TLS yet
+	if (OutConfig.HordeConnectionMode == EUbaHordeConnectionMode::Relay && OutConfig.HordeEncryption == EUbaHordeEncryption::None)
+	{
+		OutConfig.HordeEncryption = EUbaHordeEncryption::AES;
+	}
+
+	if (OutConfig.IsHordeEnabled())
+	{
+		UE_LOG(
+			LogUbaHorde, Log,
+			TEXT("UBA/Horde Configuration:")
+			TEXT("\n  Server:            %s")
+			TEXT("\n  Pool:              %s")
+			TEXT("\n  Condition:         %s")
+			TEXT("\n  Cluster:           %s")
+			TEXT("\n  Host:              %s")
+			TEXT("\n  MaxCores:          %d")
+			TEXT("\n  Allow Wine:        %s")
+			TEXT("\n  Mode:              %s")
+			TEXT("\n  Encryption:        %s")
+			, *OutConfig.HordeServer
+			, *OutConfig.GetHordePool()
+			, *OutConfig.HordeCondition
+			, *OutConfig.HordeCluster
+			, *OutConfig.HordeHost
+			, OutConfig.HordeMaxCores
+			, OutConfig.bHordeAllowWine ? TEXT("Yes") : TEXT("No")
+			, LexToString(OutConfig.HordeConnectionMode)
+			, LexToString(OutConfig.HordeEncryption)
+		);
+	}
+	else
+	{
+		UE_LOG(LogUbaHorde, Log, TEXT("UBA/Horde Configuration: Not Enabled"));
+	}
+}
+
+UBACOORDINATORHORDE_API bool FUbaHordeConfig::IsHordeEnabled() const
+{
+	return HordeEnabled == EUbaHordeEnabledMode::Enabled || (HordeEnabled == EUbaHordeEnabledMode::EnabledForBuildMachineOnly && GIsBuildMachine);
+}
+
+UBACOORDINATORHORDE_API const FUbaHordeConfig& FUbaHordeConfig::Get()
+{
+	static FUbaHordeConfig Config;
+	static FRWLock InitializedLock;
+	static bool bInitialized = false;
+
+	{
+		FRWScopeLock LockGuard(InitializedLock, FRWScopeLockType::SLT_ReadOnly);
+		if (bInitialized)
+		{
+			return Config;
+		}
+	}
+
+	{
+		FRWScopeLock LockGuard(InitializedLock, FRWScopeLockType::SLT_Write);
+		if (!bInitialized)
+		{
+			ReadUbaHordeConfig(Config);
+			bInitialized = true;
+		}
+	}
+
+	return Config;
+}
+
+UBACOORDINATORHORDE_API const TCHAR* LexToString(EUbaHordeConnectionMode InConnectionMode)
+{
+	switch (InConnectionMode)
+	{
+	case EUbaHordeConnectionMode::Direct:	return TEXT("direct");
+	case EUbaHordeConnectionMode::Tunnel:	return TEXT("tunnel");
+	case EUbaHordeConnectionMode::Relay:	return TEXT("relay");
+	}
+	return TEXT("");
+}
+
+UBACOORDINATORHORDE_API bool LexFromString(EUbaHordeConnectionMode& OutConnectionMode, const TCHAR* InString)
+{
+	if (FCString::Stricmp(InString, TEXT("direct")) == 0)
+	{
+		OutConnectionMode = EUbaHordeConnectionMode::Direct;
+	}
+	else if (FCString::Stricmp(InString, TEXT("tunnel")) == 0)
+	{
+		OutConnectionMode = EUbaHordeConnectionMode::Tunnel;
+	}
+	else if (FCString::Stricmp(InString, TEXT("relay")) == 0)
+	{
+		OutConnectionMode = EUbaHordeConnectionMode::Relay;
+	}
+	else
+	{
+		return false;
+	}
+	return true;
+}
+
+UBACOORDINATORHORDE_API const TCHAR* LexToString(EUbaHordeEncryption InTransportEncrypted)
+{
+	switch (InTransportEncrypted)
+	{
+	case EUbaHordeEncryption::None:	break;
+	case EUbaHordeEncryption::AES: return TEXT("aes");
+	}
+	return TEXT("");
+}
+
+UBACOORDINATORHORDE_API bool LexFromString(EUbaHordeEncryption& OutEncryption, const TCHAR* InString)
+{
+	if (FCString::Stricmp(InString, TEXT("NONE")) == 0)
+	{
+		OutEncryption = EUbaHordeEncryption::None;
+	}
+	else if (FCString::Stricmp(InString, TEXT("AES")) == 0)
+	{
+		OutEncryption = EUbaHordeEncryption::AES;
+	}
+	else
+	{
+		return false;
+	}
+	return true;
+}
+
+UBACOORDINATORHORDE_API const TCHAR* LexToString(EUbaHordeEnabledMode InEnabledMode)
+{
+	switch (InEnabledMode)
+	{
+	case EUbaHordeEnabledMode::Enabled: return TEXT("True");
+	case EUbaHordeEnabledMode::Disabled: return TEXT("False");
+	case EUbaHordeEnabledMode::EnabledForBuildMachineOnly: return TEXT("BuildMachineOnly");
+	}
+	return TEXT("");
+}
+
+UBACOORDINATORHORDE_API bool LexFromString(EUbaHordeEnabledMode& OutEnabledMode, const TCHAR* InString)
+{
+	if (FCString::Stricmp(InString, TEXT("True")) == 0)
+	{
+		OutEnabledMode = EUbaHordeEnabledMode::Enabled;
+	}
+	else if (FCString::Stricmp(InString, TEXT("False")) == 0)
+	{
+		OutEnabledMode = EUbaHordeEnabledMode::Disabled;
+	}
+	else if (FCString::Stricmp(InString, TEXT("BuildMachineOnly")) == 0)
+	{
+		OutEnabledMode = EUbaHordeEnabledMode::EnabledForBuildMachineOnly;
+	}
+	else
+	{
+		return false;
+	}
+	return true;
+}
+
+
diff --git a/Engine/Source/Developer/UbaCoordinatorHorde/Private/UbaHordeMetaClient.cpp b/Engine/Source/Developer/UbaCoordinatorHorde/Private/UbaHordeMetaClient.cpp
--- a/Engine/Source/Developer/UbaCoordinatorHorde/Private/UbaHordeMetaClient.cpp
+++ b/Engine/Source/Developer/UbaCoordinatorHorde/Private/UbaHordeMetaClient.cpp
@@ -1,11 +1,13 @@
 // Copyright Epic Games, Inc. All Rights Reserved.
 
 #include "UbaHordeMetaClient.h"
+#include "UbaHordeConfig.h"
 #include "Horde.h"
 #include "Serialization/JsonReader.h"
 #include "Serialization/JsonSerializer.h"
 #include "Misc/App.h"
 #include "Misc/ConfigCacheIni.h"
+#include "Misc/ScopeExit.h"
 #include "HttpServerConstants.h"
 
 DEFINE_LOG_CATEGORY(LogUbaHorde);
@@ -35,56 +37,133 @@ bool FUbaHordeMetaClient::RefreshHttpClient()
 	return true;
 }
 
-static const TCHAR* LexToString(EHordeTransportEncryption InTransportEncrypted)
+FString FUbaHordeMetaClient::BuildHordeRequestJsonBody(const FString& PoolId, EUbaHordeConnectionMode ConnectionMode, EUbaHordeEncryption Encryption, bool bExclusiveAccess)
 {
-	switch (InTransportEncrypted)
+	// Construct JSON body for Horde request, e.g. '{ "connection": { "modePreference": "relay", "encryption": "aes" } }'
+	const TSharedPtr<FJsonObject> RequestJNode = MakeShared<FJsonObject>();
 	{
-	case EHordeTransportEncryption::None:	break;
-	case EHordeTransportEncryption::AES:	return TEXT("aes");
+		if (!PoolId.IsEmpty() || bExclusiveAccess)
+		{
+			const TSharedPtr<FJsonObject> RequirementsJNode = MakeShared<FJsonObject>();
+			{
+				if (!PoolId.IsEmpty())
+				{
+					RequirementsJNode->SetStringField(TEXT("pool"), PoolId);
+				}
+				RequirementsJNode->SetBoolField(TEXT("exclusive"), bExclusiveAccess);
+			}
+			RequestJNode->SetObjectField(TEXT("requirements"), RequirementsJNode);
+		}
+
+		const TSharedPtr<FJsonObject> ConnectionJNode = MakeShared<FJsonObject>();
+		{
+			ConnectionJNode->SetStringField(TEXT("modePreference"), LexToString(ConnectionMode));
+			if (Encryption != EUbaHordeEncryption::None)
+			{
+				ConnectionJNode->SetStringField(TEXT("encryption"), LexToString(Encryption));
+			}
+
+			const TSharedPtr<FJsonObject> PortsJNode = MakeShared<FJsonObject>();
+			{
+				PortsJNode->SetNumberField(TEXT("UbaPort"), (double)FHordeRemoteMachineInfo::UbaPort);
+				PortsJNode->SetNumberField(TEXT("UbaProxyPort"), (double)FHordeRemoteMachineInfo::UbaProxyPort);
+			}
+			ConnectionJNode->SetObjectField(TEXT("ports"), PortsJNode);
+		}
+		RequestJNode->SetObjectField(TEXT("connection"), ConnectionJNode);
 	}
-	return TEXT("");
+
+	FString JsonContent;
+	const TSharedRef<TJsonWriter<>> JsonWriter = TJsonWriterFactory<>::Create(&JsonContent);
+	FJsonSerializer::Serialize(RequestJNode.ToSharedRef(), JsonWriter);
+	return JsonContent;
 }
 
-static void BuildHordeRequestJsonBody(TStringBuilder<1024>& RequestJsonBody, const FString& InPoolId, EHordeTransportEncryption InTransportEncrypted)
+TSharedPtr<FUbaHordeMetaClient::HordeClusterPromise, ESPMode::ThreadSafe> FUbaHordeMetaClient::RequestClusterId(const FString& HordeRequestJsonBody)
 {
-	RequestJsonBody.AppendChar(TEXT('{'));
+	TSharedPtr<HordeClusterPromise> Promise = MakeShared<HordeClusterPromise, ESPMode::ThreadSafe>();
 
-	// Append requirements: pool ID and whether it should be allocated exclusively for this request
-	RequestJsonBody.Appendf(TEXT("\"requirements\":{\"pool\":\"%s\",\"exclusive\":true}"), *InPoolId);
+	FHttpRequestRef Request = HttpClient->CreateRequest(TEXT("POST"), TEXT("api/v2/compute/_cluster"));
 
-	// Append connection information if encryption is enabled
-	if (InTransportEncrypted != EHordeTransportEncryption::None)
-	{
-		RequestJsonBody.Appendf(TEXT(",\"connection\":{\"encryption\":\"%s\",\"modePreference\":\"Direct\"}"), LexToString(InTransportEncrypted));
-	}
+	Request->SetDelegateThreadPolicy(EHttpRequestDelegateThreadPolicy::CompleteOnHttpThread);
+	Request->SetContentAsString(HordeRequestJsonBody);
+	Request->SetHeader(TEXT("Content-Type"), TEXT("application/json"));
+
+	Request->OnProcessRequestComplete().BindLambda(
+		[this, Promise](FHttpRequestPtr /*Request*/, FHttpResponsePtr HttpResponse, bool bConnectedSuccessfully)
+		{
+			FHordeClusterInfo Info;
+
+			ON_SCOPE_EXIT
+			{
+				Promise->SetValue(MakeTuple(HttpResponse, Info));
+			};
+
+			if (!bConnectedSuccessfully || !HttpResponse.IsValid())
+			{
+				UE_LOG(LogUbaHorde, Display, TEXT("No response from Horde"));
+				return;
+			}
 
-	RequestJsonBody.AppendChar(TEXT('}'));
+			const EHttpServerResponseCodes ResponseCode = (EHttpServerResponseCodes)HttpResponse->GetResponseCode();
+			const FString ResponseStr = HttpResponse->GetContentAsString();
+
+			if (ResponseCode == EHttpServerResponseCodes::ServiceUnavail ||
+				ResponseCode == EHttpServerResponseCodes::TooManyRequests)
+			{
+				// Service Unavailable
+				UE_LOG(LogUbaHorde, Display, TEXT("Horde agent request returned with HTTP/%d: %s"), ResponseCode, *ResponseStr);
+				return;
+			}
+
+			TSharedPtr<FJsonValue> OutJson;
+			TSharedRef<TJsonReader<>> Reader = TJsonReaderFactory<>::Create(ResponseStr);
+
+			if (!FJsonSerializer::Deserialize(Reader, OutJson, FJsonSerializer::EFlags::None))
+			{
+				// Report invalid response body with Display verbosity only, since this should not fail a CIS job
+				UE_LOG(LogUbaHorde, Display, TEXT("Invalid response body: %s"), *ResponseStr);
+				return;
+			}
+
+			TSharedPtr<FJsonValue> ClusterIdValue = OutJson->AsObject()->TryGetField(TEXT("clusterId"));
+			if (!ClusterIdValue.Get())
+			{
+				// Report invalid response body with Display verbosity only, since this should not fail a CIS job
+				UE_LOG(LogUbaHorde, Display, TEXT("Invalid response body: %s"), *ResponseStr);
+				return;
+			}
+
+			// Successfully return resolved cluster ID
+			Info.ClusterId = ClusterIdValue->AsString();
+		});
+
+	Request->ProcessRequest();
+
+	return Promise;
 }
 
-TSharedPtr<FUbaHordeMetaClient::HordeMachinePromise, ESPMode::ThreadSafe> FUbaHordeMetaClient::RequestMachine(const FString& PoolId, const TCHAR* Machine, EHordeTransportEncryption Encryption)
+TSharedPtr<FUbaHordeMetaClient::HordeMachinePromise, ESPMode::ThreadSafe> FUbaHordeMetaClient::RequestMachine(const FString& HordeRequestJsonBody, const TCHAR* ClusterId)
 {
 	TSharedPtr<HordeMachinePromise> Promise = MakeShared<HordeMachinePromise, ESPMode::ThreadSafe>();
 
-	const FString ResourcePath = FString::Format(TEXT("api/v2/compute/{0}"), { Machine != nullptr ? Machine : TEXT("default") });
-
+	const FString ResourcePath = FString::Format(TEXT("api/v2/compute/{0}"), { ClusterId != nullptr && *ClusterId != TEXT('\0') ? ClusterId : TEXT("default")});
 	FHttpRequestRef Request = HttpClient->CreateRequest(TEXT("POST"), *ResourcePath);
 
 	Request->SetDelegateThreadPolicy(EHttpRequestDelegateThreadPolicy::CompleteOnHttpThread);
-
-	// Post JSON document with constraints to acquire a Horde agent. Use pool of agents (e.g. "BoxLinux" or "BoxWin")
-	// and require exclusive access or UbaStorage will fail to initialize the next time we connect to the same machine.
-	TStringBuilder<1024> RequestJsonBody;
-	BuildHordeRequestJsonBody(RequestJsonBody, PoolId, Encryption);
-
-	Request->SetContentAsString(RequestJsonBody.ToString());
+	Request->SetContentAsString(HordeRequestJsonBody);
 	Request->SetHeader(TEXT("Content-Type"), TEXT("application/json"));
 
-	UE_LOG(LogUbaHorde, Log, TEXT("Requesting Horde agent with JSON descriptor: '%s'"), RequestJsonBody.ToString());
-
 	Request->OnProcessRequestComplete().BindLambda(
-		[this, Promise, Encryption](FHttpRequestPtr /*Request*/, FHttpResponsePtr HttpResponse, bool bConnectedSuccessfully)
+		[this, Promise, ClusterId = FString(ClusterId)](FHttpRequestPtr /*Request*/, FHttpResponsePtr HttpResponse, bool bConnectedSuccessfully)
 		{
 			FHordeRemoteMachineInfo Info;
+
+			ON_SCOPE_EXIT
+			{
+				Promise->SetValue(MakeTuple(HttpResponse, Info));
+			};
+
 			Info.Ip = TEXT("");
 			Info.Port = 0xFFFF;
 			Info.bRunsWindowOS = false;
@@ -93,7 +172,6 @@ TSharedPtr<FUbaHordeMetaClient::HordeMachinePromise, ESPMode::ThreadSafe> FUbaHo
 			if (!bConnectedSuccessfully || !HttpResponse.IsValid())
 			{
 				UE_LOG(LogUbaHorde, Display, TEXT("No response from Horde"));
-				Promise->SetValue(MakeTuple(HttpResponse, Info));
 				return;
 			}
 
@@ -105,7 +183,6 @@ TSharedPtr<FUbaHordeMetaClient::HordeMachinePromise, ESPMode::ThreadSafe> FUbaHo
 			{
 				// Service Unavailable
 				UE_LOG(LogUbaHorde, Display, TEXT("Horde agent request returned with HTTP/%d: %s"), ResponseCode, *ResponseStr);
-				Promise->SetValue(MakeTuple(HttpResponse, Info));
 				return;
 			}
 
@@ -116,7 +193,6 @@ TSharedPtr<FUbaHordeMetaClient::HordeMachinePromise, ESPMode::ThreadSafe> FUbaHo
 			{
 				// Report invalid response body with Display verbosity only, since this should not fail a CIS job
 				UE_LOG(LogUbaHorde, Display, TEXT("Invalid response body: %s"), *ResponseStr);
-				Promise->SetValue(MakeTuple(HttpResponse, Info));
 				return;
 			}
 
@@ -128,10 +204,38 @@ TSharedPtr<FUbaHordeMetaClient::HordeMachinePromise, ESPMode::ThreadSafe> FUbaHo
 			{
 				// Report invalid response body with Display verbosity only, since this should not fail a CIS job
 				UE_LOG(LogUbaHorde, Display, TEXT("Invalid response body: %s"), *ResponseStr);
-				Promise->SetValue(MakeTuple(HttpResponse, Info));
 				return;
 			}
 
+			// Check for optional port mapping array
+			if (TSharedPtr<FJsonValue> PortsValue = OutJson->AsObject()->TryGetField(TEXT("ports")))
+			{
+				for (auto Iter = PortsValue->AsObject()->Values.CreateConstIterator(); Iter; ++Iter)
+				{
+					FHordeRemoteMachineInfo::FPortInfo& PortInfo = Info.Ports.Add(Iter.Key());
+					if (TSharedPtr<FJsonValue> PortInfoPortValue = Iter.Value()->AsObject()->TryGetField(TEXT("port")))
+					{
+						PortInfo.Port = (uint16)PortInfoPortValue->AsNumber();
+					}
+					if (TSharedPtr<FJsonValue> PortInfoAgentPortValue = Iter.Value()->AsObject()->TryGetField(TEXT("agentPort")))
+					{
+						PortInfo.AgentPort = (uint16)PortInfoAgentPortValue->AsNumber();
+					}
+				}
+			}
+
+			// Check for connection mode and address - "connectionAddress" field is meaningless without a mode other than the default one
+			if (TSharedPtr<FJsonValue> ConnectionModeValue = OutJson->AsObject()->TryGetField(TEXT("connectionMode")))
+			{
+				if (LexFromString(Info.ConnectionMode, *ConnectionModeValue->AsString()))
+				{
+					if (TSharedPtr<FJsonValue> ConnectionAddressValue = OutJson->AsObject()->TryGetField(TEXT("connectionAddress")))
+					{
+						Info.ConnectionAddress = ConnectionAddressValue->AsString();
+					}
+				}
+			}
+
 			FString OsFamily(TEXT("UNKNOWN-OS"));
 
 			uint16 LogicalCores = 0;
@@ -180,29 +284,38 @@ TSharedPtr<FUbaHordeMetaClient::HordeMachinePromise, ESPMode::ThreadSafe> FUbaHo
 			FString IpString = IpValue->AsString();
 			uint16 PortNumber = (uint16)PortValue->AsNumber();
 
-			FString HordeMachineSummary = FString::Printf(TEXT("UBA Horde machine assigned (%s) [%s:%u]"), *OsFamily, *IpString, (uint32)PortNumber);
+			// Return final response information
+			Info.Ip = IpString;
+			Info.Port = PortNumber;
+			FString::ToHexBlob(NonceString, Info.Nonce, HORDE_NONCE_SIZE);
+
+			FString HordeMachineSummary = FString::Printf(TEXT("UBA Horde machine assigned (%s) on '%s' [%s:%u]"), *OsFamily, *ClusterId, *Info.GetConnectionAddress(), (uint32)Info.GetConnectionPort().Port);
 
 #if WITH_SSL
-			if (Encryption == EHordeTransportEncryption::AES)
+			if (TSharedPtr<FJsonValue> EncryptionValue = OutJson->AsObject()->TryGetField(TEXT("encryption")))
 			{
-				TSharedPtr<FJsonValue> KeyValue = OutJson->AsObject()->TryGetField(TEXT("key"));
-				FString KeyString = KeyValue ? KeyValue->AsString() : FString();
-				if (KeyString.IsEmpty())
-				{
-					HordeMachineSummary += TEXT(" [AES key missing]");
-				}
-				else if (KeyString.Len() != sizeof(Info.Key) * 2)
+				if (LexFromString(Info.Encryption, *EncryptionValue->AsString()))
 				{
-					HordeMachineSummary += TEXT(" [AES key corrupted]");
-				}
-				else
-				{
-					FString::ToHexBlob(KeyString, Info.Key, sizeof(Info.Key));
-					HordeMachineSummary += TEXT(" [AES key received]");
+					if (Info.Encryption == EUbaHordeEncryption::AES)
+					{
+						TSharedPtr<FJsonValue> KeyValue = OutJson->AsObject()->TryGetField(TEXT("key"));
+						FString KeyString = KeyValue ? KeyValue->AsString() : FString();
+						if (KeyString.IsEmpty())
+						{
+							HordeMachineSummary += TEXT(" [AES key missing]");
+						}
+						else if (KeyString.Len() != sizeof(Info.Key) * 2)
+						{
+							HordeMachineSummary += TEXT(" [AES key corrupted]");
+						}
+						else
+						{
+							FString::ToHexBlob(KeyString, Info.Key, sizeof(Info.Key));
+							HordeMachineSummary += TEXT(" [AES key received]");
+						}
+					}
 				}
 			}
-
-			Info.Encryption = Encryption;
 #endif
 
 			if (TSharedPtr<FJsonValue> LeaseIdValue = OutJson->AsObject()->TryGetField(TEXT("leaseId")))
@@ -211,13 +324,6 @@ TSharedPtr<FUbaHordeMetaClient::HordeMachinePromise, ESPMode::ThreadSafe> FUbaHo
 			}
 
 			UE_LOG(LogUbaHorde, Display, TEXT("%s"), *HordeMachineSummary);
-
-			// Return final response information
-			Info.Ip = IpString;
-			Info.Port = PortNumber;
-			FString::ToHexBlob(NonceString, Info.Nonce, HORDE_NONCE_SIZE);
-
-			Promise->SetValue(MakeTuple(HttpResponse, Info));
 		});
 
 	Request->ProcessRequest();
diff --git a/Engine/Source/Developer/UbaCoordinatorHorde/Private/UbaHordeMetaClient.h b/Engine/Source/Developer/UbaCoordinatorHorde/Private/UbaHordeMetaClient.h
--- a/Engine/Source/Developer/UbaCoordinatorHorde/Private/UbaHordeMetaClient.h
+++ b/Engine/Source/Developer/UbaCoordinatorHorde/Private/UbaHordeMetaClient.h
@@ -8,6 +8,7 @@
 #include "Interfaces/IHttpResponse.h"
 #include "HordeHttpClient.h"
 #include "UbaHordeAgentManager.h"
+#include "UbaHordeConfig.h"
 
 DECLARE_LOG_CATEGORY_EXTERN(LogUbaHorde, Log, Log);
 
@@ -22,14 +23,49 @@ DECLARE_LOG_CATEGORY_EXTERN(LogUbaHorde, Log, Log);
 // - Nonce will be all 0
 struct FHordeRemoteMachineInfo
 {
+	struct FPortInfo
+	{
+		FPortInfo() = default;
+		FPortInfo(uint16 InPort, uint16 InAgentPort)
+			: Port(InPort)
+			, AgentPort(InAgentPort)
+		{
+		}
+
+		uint16 Port = 0;
+		uint16 AgentPort = 0;
+	};
+
 	FString Ip;
+	EUbaHordeConnectionMode ConnectionMode = EUbaHordeConnectionMode::Direct;
+	FString ConnectionAddress; // Onyl used for relay mode
 	uint16 Port = 0;
+	TMap<FString, FPortInfo> Ports; // Maps compute service names to port/agent-port pairs, e.g. "_horde_compute" -> {Port: 12345, AgentPort: 7000}
 	uint16 LogicalCores = 0;
-	EHordeTransportEncryption Encryption = EHordeTransportEncryption::None;
+	EUbaHordeEncryption Encryption = EUbaHordeEncryption::None;
 	uint8 Nonce[HORDE_NONCE_SIZE] = {};
 	uint8 Key[HORDE_KEY_SIZE] = {};
 	TArray<uint8> CertificateData; // X.509 certificate for SSL connections in Pkcs12 format.
 	bool bRunsWindowOS = false;
+
+	static constexpr uint16 UbaPort = 7001;
+	static constexpr uint16 UbaProxyPort = 7002;
+
+	const FString& GetConnectionAddress() const
+	{
+		return ConnectionMode == EUbaHordeConnectionMode::Relay ? ConnectionAddress : Ip;
+	}
+
+	FPortInfo GetConnectionPort() const
+	{
+		static const FString HordeComputeIdentifier = TEXT("_horde_compute");
+		return ConnectionMode == EUbaHordeConnectionMode::Relay ? Ports[HordeComputeIdentifier] : FPortInfo(Port, Port);
+	}
+};
+
+struct FHordeClusterInfo
+{
+	FString ClusterId = TEXT("default");
 };
 
 // This encapsulates the mechanism of talking to the Horde "meta server". The server which
@@ -39,11 +75,16 @@ class FUbaHordeMetaClient final
 public:
 	// We return the HttpResponse in case more information could be used out of it later.
 	using HordeMachinePromise = TPromise<TTuple<FHttpResponsePtr, FHordeRemoteMachineInfo>>;
-	
+	using HordeClusterPromise = TPromise<TTuple<FHttpResponsePtr, FHordeClusterInfo>>;
+
 	FUbaHordeMetaClient() = default;
 
 	bool RefreshHttpClient();
 
+	static FString BuildHordeRequestJsonBody(const FString& PoolId, EUbaHordeConnectionMode ConnectionMode, EUbaHordeEncryption Encryption, bool bExclusiveAccess = true);
+
+	TSharedPtr<HordeClusterPromise, ESPMode::ThreadSafe> RequestClusterId(const FString& HordeRequestJsonBody);
+
 	// This will make a request to Horde for a remote machine to do work on.
 	// Example of actually getting the FHordeRemoteMachineInfo struct:
 	// ```
@@ -52,11 +93,9 @@ public:
 	// ... // Can to asynchronous work
 	// auto Future.Wait(); // Wait for the response to arrive
 	// FHordeRemoteMachineInfo MachineInfo = Future.Get().Value;
-	TSharedPtr<HordeMachinePromise, ESPMode::ThreadSafe> RequestMachine(const FString& PoolId, const TCHAR* Machine = nullptr, EHordeTransportEncryption Encryption = EHordeTransportEncryption::None);
+	TSharedPtr<HordeMachinePromise, ESPMode::ThreadSafe> RequestMachine(const FString& HordeRequestJsonBody, const TCHAR* ClusterId = nullptr);
 
 private:
-	void ParseConfig(const FString& IniFilename, const FString& IniSection);
-
 	FString ServerUrl;
 	TUniquePtr<FHordeHttpClient> HttpClient;
 };
diff --git a/Engine/Source/Developer/UbaCoordinatorHorde/Public/UbaHordeAgentManager.h b/Engine/Source/Developer/UbaCoordinatorHorde/Public/UbaHordeAgentManager.h
--- a/Engine/Source/Developer/UbaCoordinatorHorde/Public/UbaHordeAgentManager.h
+++ b/Engine/Source/Developer/UbaCoordinatorHorde/Public/UbaHordeAgentManager.h
@@ -5,33 +5,22 @@
 #include "Containers/Array.h"
 #include "HAL/Thread.h"
 #include "UbaBase.h"
+#include "UbaHordeConfig.h"
 
 class FEvent;
 class FUbaHordeMetaClient;
 namespace uba { class NetworkServer; }
 
-enum class EHordeTransportEncryption
-{
-	None,
-	AES,
-};
-
 class FUbaHordeAgentManager
 {
 public:
 	UBACOORDINATORHORDE_API FUbaHordeAgentManager(const FString& InWorkingDir, const FString& BinariesPath);
 	UBACOORDINATORHORDE_API ~FUbaHordeAgentManager();
 
-	void SetPool(const FString& InPool) { Pool = InPool; }
-	void SetMaxCoreCount(uint32 Count) { MaxCores = Count; }
-	void SetEncryption(EHordeTransportEncryption InEncryption) { Encryption = InEncryption; }
-	void SetUbaHost(const FAnsiString& InHost) { UbaHost = InHost; }
-	void SetUbaPort(uint32 InPort) { UbaPort = InPort; }
-
 	UBACOORDINATORHORDE_API void SetTargetCoreCount(uint32 Count);
 
-	using AddClientCallback = bool(void* userData, const uba::tchar* ip, uint16 port);
-	UBACOORDINATORHORDE_API void SetAddClientCallback(AddClientCallback* callback, void* userData);
+	using FAddClientCallback = bool(void* UserData, const uba::tchar* Ip, uint16 Port, const uba::tchar* Crypto16Characters);
+	UBACOORDINATORHORDE_API void SetAddClientCallback(FAddClientCallback* Callback, void* UserData);
 
 	// Returns the number of agents currently handled by this agent manager.
 	UBACOORDINATORHORDE_API int32 GetAgentCount() const;
@@ -52,12 +41,6 @@ private:
 	FString WorkingDir;
 	FString BinariesPath;
 
-	FString Pool;
-	FAnsiString UbaHost;
-	uint32 UbaPort = 7001;
-	uint32 MaxCores = 500; // Assume 500 by default in case this configuration is missing but a UBA pool was found
-	EHordeTransportEncryption Encryption = EHordeTransportEncryption::None;
-
 	TUniquePtr<FUbaHordeMetaClient> HordeMetaClient;
 
 	FCriticalSection BundleRefPathsLock;
@@ -72,8 +55,8 @@ private:
 	TAtomic<uint32> ActiveCoreCount;
 	TAtomic<bool> bAskForAgents;
 
-	AddClientCallback* m_callback = nullptr;
-	void* m_userData = nullptr;
+	FAddClientCallback* AddClientCallback = nullptr;
+	void* AddClientUserData = nullptr;
 
 	FUbaHordeAgentManager(const FUbaHordeAgentManager&) = delete;
 	void operator=(const FUbaHordeAgentManager&) = delete;
diff --git a/Engine/Source/Developer/UbaCoordinatorHorde/Public/UbaHordeConfig.h b/Engine/Source/Developer/UbaCoordinatorHorde/Public/UbaHordeConfig.h
new file mode 100644
--- /dev/null
+++ b/Engine/Source/Developer/UbaCoordinatorHorde/Public/UbaHordeConfig.h
@@ -0,0 +1,106 @@
+// Copyright Epic Games, Inc. All Rights Reserved.
+#pragma once
+
+#include "Containers/UnrealString.h"
+#include "Containers/Array.h"
+
+class FEvent;
+class FUbaHordeMetaClient;
+namespace uba { class NetworkServer; }
+
+enum class EUbaHordeEncryption
+{
+	None,
+	AES,
+};
+
+enum class EUbaHordeConnectionMode
+{
+	Direct,
+	Tunnel,
+	Relay,
+};
+
+enum class EUbaHordeEnabledMode
+{
+	Disabled,
+	EnabledForBuildMachineOnly,
+	Enabled,
+};
+
+/** UBA/Horde configuration structure. This is meant to match the C# version UnrealBuildTool.UnrealBuildAcceleratorHordeConfig. */
+struct FUbaHordeConfig
+{
+	static const TCHAR* ClusterDefault;
+	static const TCHAR* ClusterAuto;
+
+	/** URI of the Horde server. */
+	FString HordeServer;
+
+	/** Authentication token for the Horde server. */
+	FString HordeToken;
+
+	/** Default Horde pool name. */
+	FString DefaultHordePool;
+
+	/** Overrides the default pool name when no platform specific pool is specified. */
+	FString OverrideHordePool;
+
+	/** Horde pool name for Windows target platform. */
+	FString WindowsHordePool;
+
+	/** Horde pool name for Mac target platform. */
+	FString MacHordePool;
+
+	/** Horde pool name for Linux target platform. */
+	FString LinuxHordePool;
+
+	/** Requirements for the Horde agent to assign. */
+	FString HordeCondition;
+
+	/** Compute cluster ID to use in Horde. Set to "_auto" to let Horde server resolve a suitable cluster. */
+	FString HordeCluster = FUbaHordeConfig::ClusterAuto;
+
+	/** Which IP address UBA server should assign to agents. */
+	FString HordeHost;
+
+	/** Optional sentry URL to send UBA data to. */
+	FString UbaSentryUrl;
+
+	/** Maxmium number of CPU cores allowed to be used by build session. */
+	int32 HordeMaxCores = 576;
+
+	/** Maximum number of local CPU cores allowed to be used by build session. By default -1, which indicates to use as many CPU cores as the local machine can provide. */
+	int32 HordeMaxLocalCores = -1;
+
+	/** Allow use of POSIX/Wine. Only applicable to Horde agents running Linux. Can still be ignored if Wine executable is not set on agent. */
+	bool bHordeAllowWine = true;
+
+	/** Specifies how to connect to the remote machine. Relay mode implies encrypted connections. */
+	EUbaHordeConnectionMode HordeConnectionMode = EUbaHordeConnectionMode::Direct;
+
+	/** Specifies transport layer encryption. Currently only AES encryption is supported. */
+	EUbaHordeEncryption HordeEncryption = EUbaHordeEncryption::None;
+
+	/** Specifies whether UBA/Horde is enabled for the current session. By default disabled. */
+	EUbaHordeEnabledMode HordeEnabled = EUbaHordeEnabledMode::Disabled;
+	
+	/** Returns the active Horde pool depending on override settings and current desktop platform. */
+	UBACOORDINATORHORDE_API const FString& GetHordePool() const;
+
+	/** Returns true if UBA/Horde is enabled. This depends on HordeEnabled and whether the process is running on a build machine. */
+	UBACOORDINATORHORDE_API bool IsHordeEnabled() const;
+
+	/** Returns the global UBA/Horde configuration and initializes it from a list of XML files on the first call. */
+	UBACOORDINATORHORDE_API static const FUbaHordeConfig& Get();
+};
+
+UBACOORDINATORHORDE_API const TCHAR* LexToString(EUbaHordeConnectionMode InConnectionMode);
+UBACOORDINATORHORDE_API bool LexFromString(EUbaHordeConnectionMode& OutConnectionMode, const TCHAR* InString);
+
+UBACOORDINATORHORDE_API const TCHAR* LexToString(EUbaHordeEncryption InEncryption);
+UBACOORDINATORHORDE_API bool LexFromString(EUbaHordeEncryption& OutEncryption, const TCHAR* InString);
+
+UBACOORDINATORHORDE_API const TCHAR* LexToString(EUbaHordeEnabledMode InEnabledMode);
+UBACOORDINATORHORDE_API bool LexFromString(EUbaHordeEnabledMode& OutEnabledMode, const TCHAR* InString);
+
diff --git a/Engine/Source/Developer/UbaCoordinatorHorde/UbaCoordinatorHorde.Build.cs b/Engine/Source/Developer/UbaCoordinatorHorde/UbaCoordinatorHorde.Build.cs
--- a/Engine/Source/Developer/UbaCoordinatorHorde/UbaCoordinatorHorde.Build.cs
+++ b/Engine/Source/Developer/UbaCoordinatorHorde/UbaCoordinatorHorde.Build.cs
@@ -22,6 +22,7 @@ public class UbaCoordinatorHorde : ModuleRules
 			"Sockets",
 			"SSL",
 			"OpenSSL",
+			"XmlParser",
 		});
 
 		PrivateIncludePathModuleNames.AddRange(new string[] {
