[Backout] - CL39349947
#rnx
Original CL Desc
-----------------------------------------------------------------
[UBA] Refactored UBA configuration and added support for Relay mode for shader compilation.

This contains two significant changes:
1. Refactored how UBA/Horde configuration is read in C++ for shader compilation to match C# code for UBT. This is implemented in FUbaHordeConfig and applies the same list of BuildConfiguration.xml files in the same order as UBT. It also adds the same command line arguments. This effectively replaces the old [UbaController] configurations from BaseEngine.ini. If a machine is set up for UBA in UBT, it will now automatically be set up for UBA in shader compilation.
2. Add support for Relay mode. This required the addition of "_auto" Horde cluster to be resolved via an HTTP/Json request prior to the agent assignment request. This allows the client to get a Horde cluster assigned automatically based on the user's IP address (i.e. region based).

diff --git a/Engine/Plugins/UbaController/Source/Private/UbaControllerModule.cpp b/Engine/Plugins/UbaController/Source/Private/UbaControllerModule.cpp
--- a/Engine/Plugins/UbaController/Source/Private/UbaControllerModule.cpp
+++ b/Engine/Plugins/UbaController/Source/Private/UbaControllerModule.cpp
@@ -4,7 +4,6 @@
 #include "HttpModule.h"
 #include "Misc/ConfigCacheIni.h"
 #include "UbaJobProcessor.h"
-#include "UbaHordeConfig.h"
 #include "Features/IModularFeatures.h"
 #include "HAL/FileManager.h"
 #include "Misc/CommandLine.h"
@@ -91,7 +90,32 @@ static bool IsUbaControllerEnabled()
 	{
 		return false;
 	}
-	return true;
+
+	// Check if UbaController is enabled via command line argument
+	if (FParse::Param(FCommandLine::Get(), TEXT("Uba")))
+	{
+		return true;
+	}
+
+	// Check if UbaController is enabled via INI configuration in [UbaController] section.
+	FString EnabledState;
+	GConfig->GetString(TEXT("UbaController"), TEXT("Enabled"), EnabledState, GEngineIni);
+
+	// This "Enabled" parameter is a tri-state so we have to parse it as a string.
+	// Those strings can include the comments (starting with ';') from the INI file, so we have to trim that section from the string.
+	const int32 EnabledStateCommentPosition = EnabledState.Find(TEXT(";"));
+	if (EnabledStateCommentPosition != INDEX_NONE)
+	{
+		EnabledState.RemoveAt(EnabledStateCommentPosition, EnabledState.Len() - EnabledStateCommentPosition);
+		EnabledState.RemoveSpacesInline();
+	}
+
+	if (EnabledState.Equals(TEXT("True"), ESearchCase::IgnoreCase) || (EnabledState.Equals(TEXT("BuildMachineOnly"), ESearchCase::IgnoreCase) && GIsBuildMachine))
+	{
+		return true;
+	}
+
+	return false;
 }
 
 bool FUbaControllerModule::IsSupported()
diff --git a/Engine/Plugins/UbaController/Source/Private/UbaJobProcessor.cpp b/Engine/Plugins/UbaController/Source/Private/UbaJobProcessor.cpp
--- a/Engine/Plugins/UbaController/Source/Private/UbaJobProcessor.cpp
+++ b/Engine/Plugins/UbaController/Source/Private/UbaJobProcessor.cpp
@@ -8,7 +8,6 @@
 #include "Misc/CoreMisc.h"
 #include "UbaControllerModule.h"
 #include "UbaHordeAgentManager.h"
-#include "UbaHordeConfig.h"
 #include "UbaProcessStartInfo.h"
 #include "UbaSessionServerCreateInfo.h"
 #include "UbaStringConversion.h"
@@ -155,7 +154,11 @@ FUbaJobProcessor::FUbaJobProcessor(
 			checkf(false, TEXT("%s"), UBASTRING_TO_TCHAR(text));
 		});
 
-	MaxLocalParallelJobs = FUbaHordeConfig::Get().HordeMaxLocalCores;
+	if (!GConfig->GetInt(TEXT("UbaController"), TEXT("MaxLocalParallelJobs"), MaxLocalParallelJobs, GEngineIni))
+	{
+		MaxLocalParallelJobs = -1;
+	}
+
 	if (MaxLocalParallelJobs == -1)
 	{
 		MaxLocalParallelJobs = FPlatformMisc::NumberOfCoresIncludingHyperthreads();
@@ -320,6 +323,23 @@ void FUbaJobProcessor::RunTaskWithUba(FDistributedBuildTask* Task)
 
 FString GetUbaBinariesPath();
 
+static bool ParseHordeTransportEncryption(EHordeTransportEncryption& OutEncryption, const TCHAR* InName)
+{
+	if (FCString::Stricmp(InName, TEXT("NONE")) == 0)
+	{
+		OutEncryption = EHordeTransportEncryption::None;
+	}
+	else if (FCString::Stricmp(InName, TEXT("AES")) == 0)
+	{
+		OutEncryption = EHordeTransportEncryption::AES;
+	}
+	else
+	{
+		return false;
+	}
+	return true;
+}
+
 void FUbaJobProcessor::StartUba()
 {
 	checkf(UbaServer == nullptr, TEXT("FUbaJobProcessor::StartUba() was called twice before FUbaJobProcessor::ShutDownUba()"));
@@ -369,16 +389,80 @@ void FUbaJobProcessor::StartUba()
 		NetworkServer_StartListen(UbaServer, uba::DefaultPort, nullptr); // Start listen so any helper on the LAN can join in
 	}
 
-	// Only request Horde agents if Horde is enabled for UBA
-	if (FUbaHordeConfig::Get().IsHordeEnabled())
+	HordeAgentManager = MakeUnique<FUbaHordeAgentManager>(ControllerModule.GetWorkingDirectory(), GetUbaBinariesPath());
+
+	auto AddClientCallback = [](void* userData, const uba::tchar* ip, uint16 port)
+		{
+			return NetworkServer_AddClient((uba::NetworkServer*)userData, ip, port, nullptr);
+		};
+	HordeAgentManager->SetAddClientCallback(AddClientCallback, UbaServer);
+
+	FString HordeConfig;
+	if (GConfig->GetString(TEXT("UbaController"), TEXT("Horde"), HordeConfig, GEngineIni))
 	{
-		HordeAgentManager = MakeUnique<FUbaHordeAgentManager>(ControllerModule.GetWorkingDirectory(), GetUbaBinariesPath());
+		HordeConfig.TrimStartInline();
+		HordeConfig.TrimEndInline();
+		HordeConfig.RemoveFromStart(TEXT("("));
+		HordeConfig.RemoveFromEnd(TEXT(")"));
+
+		FString Pool;
+		bool FoundPool = false;
+
+		#if PLATFORM_MAC
+		FoundPool = FParse::Value(*HordeConfig, TEXT("MacPool="), Pool);
+		#endif
+
+		if (!FoundPool)
+		{
+			FoundPool = FParse::Value(*HordeConfig, TEXT("Pool="), Pool);
+		}
+
+		if (FoundPool)
+		{
+			UE_LOG(LogUbaController, Log, TEXT("Found UBA controller Pool: \"%s\""), *Pool);
+			HordeAgentManager->SetPool(Pool);
+		}
+
+		uint32 MaxCores = 0;
+		if (FParse::Value(*HordeConfig, TEXT("MaxCores="), MaxCores))
+		{
+			UE_LOG(LogUbaController, Log, TEXT("Found UBA controller MaxCores: \"%u\""), MaxCores);
+			HordeAgentManager->SetMaxCoreCount(MaxCores);
+		}
 
-		auto AddClientCallback = [](void* UserData, const uba::tchar* Ip, uint16 Port, const uba::tchar* Crypto16Characters)
+		FString EncryptionName;
+		if (FParse::Value(*HordeConfig, TEXT("Encryption="), EncryptionName))
+		{
+			EHordeTransportEncryption Encryption = EHordeTransportEncryption::None;
+			if (ParseHordeTransportEncryption(Encryption, *EncryptionName))
 			{
-				return NetworkServer_AddClient(reinterpret_cast<uba::NetworkServer*>(UserData), Ip, Port, Crypto16Characters);
-			};
-		HordeAgentManager->SetAddClientCallback(AddClientCallback, UbaServer);
+				HordeAgentManager->SetEncryption(Encryption);
+				UE_LOG(LogUbaController, Log, TEXT("Found UBA controller Encryption: \"%s\""), *EncryptionName);
+			}
+			else
+			{
+				UE_LOG(LogUbaController, Log, TEXT("Found UBA controller Encryption with unrecognized identifier: \"%s\""), *EncryptionName);
+			}
+		}
+	}
+
+	FString Host;
+	if (GConfig->GetString(TEXT("UbaController"), TEXT("Host"), Host, GEngineIni))
+	{
+		UE_LOG(LogUbaController, Log, TEXT("Found UBA controller Host: \"%s\""), *Host);
+		HordeAgentManager->SetUbaHost(StringCast<ANSICHAR>(*Host).Get());
+	}
+
+	int32 Port = 0;
+	if (GConfig->GetInt(TEXT("UbaController"), TEXT("Port"), Port, GEngineIni))
+	{
+		UE_LOG(LogUbaController, Log, TEXT("Found UBA controller Port: \"%d\""), Port);
+		HordeAgentManager->SetUbaPort(Port);
+	}
+	
+	if (GConfig->GetBool(TEXT("UbaController"), TEXT("bForceRemote"), bForceRemote, GEngineIni))
+	{
+		UE_LOG(LogUbaController, Log, TEXT("Found UBA controller Force Remote: [%s]"), bForceRemote ? TEXT("True") : TEXT("False"));
 	}
 
 	UE_LOG(LogUbaController, Display, TEXT("Created UBA storage server: RootDir=%s"), *RootDir);
@@ -418,21 +502,21 @@ uint32 FUbaJobProcessor::Run()
 	{
 		const float ElapsedSeconds = (FPlatformTime::Cycles() - LastTimeSinceHadJobs) * FPlatformTime::GetSecondsPerCycle();
 
-		uint32 NumQueuedJobs = 0;
-		uint32 NumActiveLocalJobs = 0;
-		uint32 NumActiveRemoteJobs = 0;
-		uint32 NumFinishedJobs = 0;
+		uint32 queued = 0;
+		uint32 activeLocal = 0;
+		uint32 activeRemote = 0;
+		uint32 finished = 0;
 
 		FScopeLock lock(&bShouldProcessJobsLock);
 
 		if (UbaScheduler)
 		{
-			Scheduler_GetStats(UbaScheduler, NumQueuedJobs, NumActiveLocalJobs, NumActiveRemoteJobs, NumFinishedJobs);
+			Scheduler_GetStats(UbaScheduler, queued, activeLocal, activeRemote, finished);
 		}
-		const uint32 NumActiveJobs = NumActiveLocalJobs + NumActiveRemoteJobs;
+		uint32 active = activeLocal + activeRemote;
 
 		// We don't want to hog up Horde resources.
-		if (bShouldProcessJobs && ElapsedSeconds > UbaJobProcessorOptions::MaxTimeWithoutTasks && (NumQueuedJobs + NumActiveJobs) == 0)
+		if (bShouldProcessJobs && ElapsedSeconds > UbaJobProcessorOptions::MaxTimeWithoutTasks && (queued + active) == 0)
 		{
 			// If we're optimizing job starting, we only want to shutdown UBA if all the processes have terminated
 			bShouldProcessJobs = false;
@@ -440,7 +524,7 @@ uint32 FUbaJobProcessor::Run()
 		}
 
 		// Check if we have new tasks to process
-		if ((ControllerModule.HasTasksDispatchedOrPending() || (NumQueuedJobs + NumActiveJobs) != 0))
+		if ((ControllerModule.HasTasksDispatchedOrPending() || (queued + active) != 0))
 		{
 			if (!bShouldProcessJobs)
 			{
@@ -455,20 +539,18 @@ uint32 FUbaJobProcessor::Run()
 
 		if (bShouldProcessJobs)
 		{
-			const int32 MaxAvailableLocalCores = FPlatformMisc::NumberOfCoresIncludingHyperthreads();
-			const int32 LocalCoresToNotUse = 1 + (int32)NumActiveRemoteJobs / 30; // Use one core per 30 remote ones
-			const int32 MaxLocalCoresToUse = FMath::Clamp(MaxAvailableLocalCores - LocalCoresToNotUse, 0, MaxLocalParallelJobs);
-			const int32 MaxRemoteCoresToRequest = FMath::Max(0, (int32)(NumQueuedJobs + NumActiveJobs) - MaxLocalCoresToUse);
+			int32 MaxPossible = FPlatformMisc::NumberOfCoresIncludingHyperthreads();
+			int32 LocalCoresToNotUse = 1 + activeRemote / 30; // Use one core per 30 remote ones
+			int32 MaxLocal = FMath::Max(0, MaxPossible - LocalCoresToNotUse);
+			MaxLocal = FMath::Min(MaxLocal, MaxLocalParallelJobs);
+			Scheduler_SetMaxLocalProcessors(UbaScheduler, bForceRemote ? 0 : MaxLocal);
 
-			Scheduler_SetMaxLocalProcessors(UbaScheduler, MaxLocalCoresToUse);
+			int32 TargetCoreCount = FMath::Max(0, int32(queued + active) - MaxLocal);
 
-			if (HordeAgentManager)
-			{
-				HordeAgentManager->SetTargetCoreCount(MaxRemoteCoresToRequest);
-			}
+			HordeAgentManager->SetTargetCoreCount(TargetCoreCount);
 			
 			// TODO: Not sure this is a good idea in a cooking scenario where number of queued processes are going up and down
-			SessionServer_SetMaxRemoteProcessCount(UbaSessionServer, MaxRemoteCoresToRequest);
+			SessionServer_SetMaxRemoteProcessCount(UbaSessionServer, TargetCoreCount);
 
 			UpdateStats();
 		}
@@ -578,9 +660,7 @@ void FUbaJobProcessor::HandleTaskQueueUpdated(const FString& InputFileName)
 	{
 		FDistributedBuildTask* Task = nullptr;
 		if (!ControllerModule.PendingRequestedCompilationTasks.Dequeue(Task) || !Task)
-		{
 			break;
-		}
 		RunTaskWithUba(Task);
 	}
 }
@@ -610,12 +690,9 @@ bool FUbaJobProcessor::PollStats(FDistributedBuildStats& OutStats)
 
 void FUbaJobProcessor::UpdateStats()
 {
-	if (HordeAgentManager)
-	{
-		FScopeLock StatsLockGuard(&StatsLock);
+	FScopeLock StatsLockGuard(&StatsLock);
 
-		// Update maximum
-		Stats.MaxRemoteAgents = FMath::Max(Stats.MaxRemoteAgents, (uint32)HordeAgentManager->GetAgentCount());
-		Stats.MaxActiveAgentCores = FMath::Max(Stats.MaxActiveAgentCores, HordeAgentManager->GetActiveCoreCount());
-	}
+	// Update maximum
+	Stats.MaxRemoteAgents = FMath::Max(Stats.MaxRemoteAgents, (uint32)HordeAgentManager->GetAgentCount());
+	Stats.MaxActiveAgentCores = FMath::Max(Stats.MaxActiveAgentCores, HordeAgentManager->GetActiveCoreCount());
 }
diff --git a/Engine/Plugins/UbaController/Source/Public/UbaJobProcessor.h b/Engine/Plugins/UbaController/Source/Public/UbaJobProcessor.h
--- a/Engine/Plugins/UbaController/Source/Public/UbaJobProcessor.h
+++ b/Engine/Plugins/UbaController/Source/Public/UbaJobProcessor.h
@@ -51,7 +51,7 @@ private:
 	void UpdateStats();
 
 	/** The runnable thread */
-	FRunnableThread* Thread = nullptr;
+	FRunnableThread* Thread;
 	
 	FUbaControllerModule& ControllerModule;
 
@@ -61,14 +61,14 @@ private:
 	TAtomic<bool> bForceStop;
 	
 	/** Used to abort the current processing loop */
-	uint32 LastTimeCheckedForTasks = 0;
+	uint32 LastTimeCheckedForTasks;
 	
 	/** Used to stop the processing loop in the next main loop */
 	FCriticalSection bShouldProcessJobsLock;
-	bool bShouldProcessJobs = false;
+	bool bShouldProcessJobs;
 
 	/** Set to true when the main loop finishes*/
-	bool bIsWorkDone = false;
+	bool bIsWorkDone;
 
 	uba::NetworkServer* UbaServer = nullptr;
 	uba::StorageServer* UbaStorageServer = nullptr;
@@ -83,4 +83,7 @@ private:
 
 	FCriticalSection StatsLock;
 	FDistributedBuildStats Stats;
+
+	/** If true all UBA jobs will be run remotely */
+	bool bForceRemote = false;
 };
\ No newline at end of file
diff --git a/Engine/Source/Developer/Horde/Private/Compute/AgentMessage.cpp b/Engine/Source/Developer/Horde/Private/Compute/AgentMessage.cpp
--- a/Engine/Source/Developer/Horde/Private/Compute/AgentMessage.cpp
+++ b/Engine/Source/Developer/Horde/Private/Compute/AgentMessage.cpp
@@ -63,7 +63,7 @@ void FAgentMessageChannel::UploadFiles(const char* Path, const char* Locator)
 	FlushMessage();
 }
 
-void FAgentMessageChannel::DeleteFiles(const char* const * Paths, size_t Count)
+void FAgentMessageChannel::DeleteFiles(const char** Paths, size_t Count)
 {
 	size_t RequiredSize = 10;
 	for (size_t Idx = 0; Idx < Count; Idx++)
@@ -80,7 +80,7 @@ void FAgentMessageChannel::DeleteFiles(const char* const * Paths, size_t Count)
 	FlushMessage();
 }
 
-void FAgentMessageChannel::Execute(const char* Exe, const char* const * Args, size_t NumArgs, const char* WorkingDir, const char* const * EnvVars, size_t NumEnvVars, EExecuteProcessFlags Flags)
+void FAgentMessageChannel::Execute(const char* Exe, const char** Args, size_t NumArgs, const char* WorkingDir, const char** EnvVars, size_t NumEnvVars, EExecuteProcessFlags Flags)
 {
 	size_t RequiredSize = 50 + strlen(Exe);
 	for (size_t Idx = 0; Idx < NumArgs; Idx++)
diff --git a/Engine/Source/Developer/Horde/Public/Compute/AgentMessage.h b/Engine/Source/Developer/Horde/Public/Compute/AgentMessage.h
--- a/Engine/Source/Developer/Horde/Public/Compute/AgentMessage.h
+++ b/Engine/Source/Developer/Horde/Public/Compute/AgentMessage.h
@@ -117,10 +117,10 @@ public:
 	HORDE_API void UploadFiles(const char* Path, const char* Locator);
 
 	// Deletes files matching a set of wildcards
-	HORDE_API void DeleteFiles(const char* const * Paths, size_t Count);
+	HORDE_API void DeleteFiles(const char** Paths, size_t Count);
 
 	// Executes a process on the remote machine
-	HORDE_API void Execute(const char* Exe, const char* const * Args, size_t NumArgs, const char* WorkingDir, const char* const * EnvVars, size_t NumEnvVars, EExecuteProcessFlags Flags);
+	HORDE_API void Execute(const char* Exe, const char** Args, size_t NumArgs, const char* WorkingDir, const char** EnvVars, size_t NumEnvVars, EExecuteProcessFlags Flags);
 
 	// Writes a blob requested by the remote
 	HORDE_API void Blob(const unsigned char* Data, size_t Length);
diff --git a/Engine/Source/Developer/UbaCoordinatorHorde/Private/UbaHordeAgent.cpp b/Engine/Source/Developer/UbaCoordinatorHorde/Private/UbaHordeAgent.cpp
--- a/Engine/Source/Developer/UbaCoordinatorHorde/Private/UbaHordeAgent.cpp
+++ b/Engine/Source/Developer/UbaCoordinatorHorde/Private/UbaHordeAgent.cpp
@@ -1,7 +1,6 @@
 // Copyright Epic Games, Inc. All Rights Reserved.
 
 #include "UbaHordeAgent.h"
-#include "UbaHordeConfig.h"
 
 #include "Misc/Paths.h"
 #include "HAL/FileManager.h"
@@ -29,7 +28,7 @@ FUbaHordeAgent::FUbaHordeAgent(const FHordeRemoteMachineInfo& InMachineInfo)
 	ComputeTransport->Send(MachineInfo.Nonce, sizeof(MachineInfo.Nonce));
 
 	// If encryption protocol is requested, wrap initial transport layer around an encryption layer *after* initial nonce was sent unencrypted
-	if (InMachineInfo.Encryption == EUbaHordeEncryption::AES)
+	if (InMachineInfo.Encryption == EHordeTransportEncryption::AES)
 	{
 		ComputeTransport = MakeUnique<FUbaHordeComputeTransportAES>(MachineInfo, MoveTemp(ComputeTransport));
 		if (!ComputeTransport->IsValid())
@@ -182,7 +181,7 @@ bool FUbaHordeAgent::UploadBinaries(const FString& BundleDirectory, const char*
 	return bSuccess;
 }
 
-void FUbaHordeAgent::Execute(const char* Exe, const char* const * Args, size_t NumArgs, const char* WorkingDir, const char* const * EnvVars, size_t NumEnvVars, bool bUseWine)
+void FUbaHordeAgent::Execute(const char* Exe, const char** Args, size_t NumArgs, const char* WorkingDir, const char** EnvVars, size_t NumEnvVars, bool bUseWine)
 {
 	ChildChannel->Execute(Exe, Args, NumArgs, WorkingDir, EnvVars, NumEnvVars, bUseWine ? EExecuteProcessFlags::UseWine : EExecuteProcessFlags::None);
 }
diff --git a/Engine/Source/Developer/UbaCoordinatorHorde/Private/UbaHordeAgent.h b/Engine/Source/Developer/UbaCoordinatorHorde/Private/UbaHordeAgent.h
--- a/Engine/Source/Developer/UbaCoordinatorHorde/Private/UbaHordeAgent.h
+++ b/Engine/Source/Developer/UbaCoordinatorHorde/Private/UbaHordeAgent.h
@@ -24,7 +24,7 @@ public:
 	 * e.g. Locator = abcdef1234567890abcdef1234567890abcdef12@34567890abcdef1234567890#123 */
 	bool UploadBinaries(const FString& BundleDirectory, const char* BundleLocator);
 
-	void Execute(const char* Exe, const char* const * Args, size_t NumArgs, const char* WorkingDir, const char* const * EnvVars, size_t NumEnvVars, bool bUseWine = false);
+	void Execute(const char* Exe, const char** Args, size_t NumArgs, const char* WorkingDir, const char** EnvVars, size_t NumEnvVars, bool bUseWine = false);
 
 	void CloseConnection(); 
 
diff --git a/Engine/Source/Developer/UbaCoordinatorHorde/Private/UbaHordeAgentManager.cpp b/Engine/Source/Developer/UbaCoordinatorHorde/Private/UbaHordeAgentManager.cpp
--- a/Engine/Source/Developer/UbaCoordinatorHorde/Private/UbaHordeAgentManager.cpp
+++ b/Engine/Source/Developer/UbaCoordinatorHorde/Private/UbaHordeAgentManager.cpp
@@ -17,7 +17,6 @@
 #include "Storage/BlobWriter.h"
 #include "UbaBase.h"
 #include "UbaHordeAgent.h"
-#include "UbaDefaultConstants.h"
 #include <filesystem>
 #include <fstream>
 
@@ -54,7 +53,7 @@ FUbaHordeAgentManager::~FUbaHordeAgentManager()
 
 void FUbaHordeAgentManager::SetTargetCoreCount(uint32 Count)
 {
-	TargetCoreCount = FMath::Min((uint32)FUbaHordeConfig::Get().HordeMaxCores, Count);
+	TargetCoreCount = FMath::Min(MaxCores, Count);
 
 	while (EstimatedCoreCount < TargetCoreCount)
 	{
@@ -80,10 +79,10 @@ void FUbaHordeAgentManager::SetTargetCoreCount(uint32 Count)
 	}
 }
 
-void FUbaHordeAgentManager::SetAddClientCallback(FAddClientCallback* Callback, void* UserData)
+void FUbaHordeAgentManager::SetAddClientCallback(AddClientCallback* callback, void* userData)
 {
-	AddClientCallback = Callback;
-	AddClientUserData = UserData;
+	m_callback = callback;
+	m_userData = userData;
 }
 
 int32 FUbaHordeAgentManager::GetAgentCount() const
@@ -163,52 +162,15 @@ void FUbaHordeAgentManager::RequestAgent()
 {
 	EstimatedCoreCount += 32; // We estimate a typical agent to have 32 cores
 
-	const int32 NextAgentIndex = Agents.Num();
-
 	FScopeLock AgentsScopeLock(&AgentsLock);
 	FHordeAgentWrapper& Wrapper = *Agents.Emplace_GetRef(MakeUnique<FHordeAgentWrapper>());
 
 	Wrapper.ShouldExit = FGenericPlatformProcess::GetSynchEventFromPool(true);
-
-	Wrapper.Thread = FThread(*FString::Printf(TEXT("HordeAgent %d"), NextAgentIndex), [this, &Wrapper]() { ThreadAgent(Wrapper); });
+	Wrapper.Thread = FThread(TEXT("HordeAgent"), [this, WrapperPtr = &Wrapper]() { ThreadAgent(*WrapperPtr); });
 }
 
-// Helper class to manage raw and auto-released strings for UBA agent command line arguments.
-class FUbaAgentArguments
-{
-	TArray<const char*> RawArguments;
-	TArray<FAnsiString> ArgumentStrings;
-public:
-	FUbaAgentArguments(int32 InitialArgumentCapapcity = 16)
-	{
-		RawArguments.Reserve(InitialArgumentCapapcity);
-	}
-	void AddRaw(const char* InArgument)
-	{
-		RawArguments.Add(InArgument);
-	}
-	void Add(const FAnsiString& InArgument)
-	{
-		ArgumentStrings.Add(InArgument);
-	}
-	void Finalize()
-	{
-		RawArguments.Reserve(RawArguments.Num() + ArgumentStrings.Num());
-		for (const FAnsiString& Argument : ArgumentStrings)
-		{
-			RawArguments.Add(Argument.GetCharArray().GetData());
-		}
-	}
-	const TArray<const char*>& Get() const
-	{
-		return RawArguments;
-	}
-};
-
 void FUbaHordeAgentManager::ThreadAgent(FHordeAgentWrapper& Wrapper)
 {
-	const FUbaHordeConfig& Config = FUbaHordeConfig::Get();
-
 	FEvent& ShouldExit = *Wrapper.ShouldExit;
 	TUniquePtr<FUbaHordeAgent> Agent;
 	bool bSuccess = false;
@@ -219,16 +181,14 @@ void FUbaHordeAgentManager::ThreadAgent(FHordeAgentWrapper& Wrapper)
 		{
 			Agent->CloseConnection();
 		}
+
 		ShouldExit.Trigger();
 	};
 
 	int MachineCoreCount = 0;
 
 	// If no host is specified, we need to start the agent in listen mode
-	const bool bUseListen = Config.HordeHost.IsEmpty();
-
-	FString UbaAgentIp;
-	FHordeRemoteMachineInfo::FPortInfo UbaPort, UbaProxyPort;
+	const bool bUseListen = UbaHost.IsEmpty();
 
 	{
 		ON_SCOPE_EXIT{ EstimatedCoreCount -= 32; };
@@ -311,31 +271,10 @@ void FUbaHordeAgentManager::ThreadAgent(FHordeAgentWrapper& Wrapper)
 			}
 		}
 
-		// Initialize horde request with connection mode preference and encryption.
-		// Pool ID is superseded by cluster ID, but it still has to be provided if default cluster and direct mode is used.
-		const FString HordeRequestJsonBody = HordeMetaClient->BuildHordeRequestJsonBody(Config.HordeConnectionMode == EUbaHordeConnectionMode::Direct ? Config.GetHordePool() : TEXT(""), Config.HordeConnectionMode, Config.HordeEncryption);
-
-		// Resolve cluster ID through Horde if "_auto" was specified
-		FHordeClusterInfo ClusterInfo;
-		if (Config.HordeCluster == FUbaHordeConfig::ClusterAuto)
-		{
-			TSharedPtr<FUbaHordeMetaClient::HordeClusterPromise, ESPMode::ThreadSafe> Promise = HordeMetaClient->RequestClusterId(HordeRequestJsonBody);
-			if (!Promise)
-			{
-				UE_LOG(LogUbaHorde, Log, TEXT("Failed to resolve Horde cluster ID with HTTP/Json request: %s"), *HordeRequestJsonBody);
-				return;
-			}
-			TFuture<TTuple<FHttpResponsePtr, FHordeClusterInfo>> Future = Promise->GetFuture();
-			Future.Wait();
-			ClusterInfo = Future.Get().Value;
-			UE_CLOG(!ClusterInfo.ClusterId.IsEmpty(), LogUbaHorde, Verbose, TEXT("Received resolved cluster ID from Horde: %s"), *ClusterInfo.ClusterId);
-		}
-
-		// Request Horde machine for UBA agent
-		TSharedPtr<FUbaHordeMetaClient::HordeMachinePromise, ESPMode::ThreadSafe> Promise = HordeMetaClient->RequestMachine(HordeRequestJsonBody, *ClusterInfo.ClusterId);
+		TSharedPtr<FUbaHordeMetaClient::HordeMachinePromise, ESPMode::ThreadSafe> Promise = HordeMetaClient->RequestMachine(Pool, nullptr, Encryption);
 		if (!Promise)
 		{
-			UE_LOG(LogUbaHorde, Log, TEXT("Failed to request Horde machine from cluster: %s"), *ClusterInfo.ClusterId);
+			//UE_LOG(LogUbaHorde, Error, TEXT("Failed to create Horde bundle for UbaAgent executable: %s"), *UbaAgentFilePath);
 			return;
 		}
 		TFuture<TTuple<FHttpResponsePtr, FHordeRemoteMachineInfo>> Future = Promise->GetFuture();
@@ -343,9 +282,12 @@ void FUbaHordeAgentManager::ThreadAgent(FHordeAgentWrapper& Wrapper)
 		FHordeRemoteMachineInfo MachineInfo = Future.Get().Value;
 
 		// If the machine couldn't be assigned, just ignore this agent slot
-		if (MachineInfo.GetConnectionAddress().IsEmpty())
+		if (MachineInfo.Ip == TEXT(""))
 		{
-			UE_CLOG(!LastRequestFailTime, LogUbaHorde, Verbose, TEXT("No resources available in Horde. Will keep retrying until %u cores are used (Currently have %u)"), TargetCoreCount.Load(), ActiveCoreCount.Load());
+			if (!LastRequestFailTime)
+			{
+				UE_LOG(LogUbaHorde, Verbose, TEXT("No resources available in Horde. Will keep retrying until %u cores are used (Currently have %u)"), TargetCoreCount.Load(), ActiveCoreCount.Load());
+			}
 			LastRequestFailTime = FPlatformTime::Cycles64();
 			return;
 		}
@@ -394,47 +336,20 @@ void FUbaHordeAgentManager::ThreadAgent(FHordeAgentWrapper& Wrapper)
 			}
 		}
 
-		// Get IP address and port for UbaAgent
-		UbaAgentIp = Agent->GetMachineInfo().GetConnectionAddress();
-
-		UbaPort = Agent->GetMachineInfo().Ports[TEXT("UbaPort")];
-		UbaProxyPort = Agent->GetMachineInfo().Ports[TEXT("UbaProxyPort")];
-
-		const uint32 UbaHostPort = uba::DefaultPort;
-
 		// Start the UBA Agent that will connect to us, requesting for work
-		FUbaAgentArguments UbaAgentArgs;
 
-		if (bUseListen)
-		{
-			UbaAgentArgs.Add(FAnsiString::Printf("-Listen=%u", UbaPort.AgentPort));
-			UbaAgentArgs.Add(TEXT("-ListenTimeout=10"));
-		}
-		else
-		{
-			UbaAgentArgs.Add(FAnsiString::Printf("-Host=%s:%u", TCHAR_TO_ANSI(*Config.HordeHost), UbaHostPort));
-		}
+		const FAnsiString AgentConnectionArg = bUseListen ? FAnsiString::Printf("-listen=%u", UbaPort) : FAnsiString::Printf("-Host=%s:%u", *UbaHost, UbaPort);
 
-		if (!Config.UbaSentryUrl.IsEmpty())
+		const char* UbaAgentArgs[] =
 		{
-			UbaAgentArgs.Add(FAnsiString::Printf("-Sentry=%s", TCHAR_TO_ANSI(*Config.UbaSentryUrl)));
-		}
-
-		UbaAgentArgs.Add(FAnsiString::Printf("-ProxyPort=%u", UbaProxyPort.AgentPort));
-
-		// -nopoll recommended when running on remote Horde agents to make sure they exit after completion. Otherwise, it keeps running.
-		UbaAgentArgs.AddRaw("-NoPoll");
-
-		// Skip all the agent logging that would be sent over to here
-		UbaAgentArgs.AddRaw("-Quiet");
-
-		// After 15 seconds of idling agent will automatically disconnect
-		UbaAgentArgs.AddRaw("-MaxIdle=15");
-
-		UbaAgentArgs.AddRaw("-Dir=%UE_HORDE_SHARED_DIR%\\Uba");
-		UbaAgentArgs.AddRaw("-Eventfile=%UE_HORDE_TERMINATION_SIGNAL_FILE%");
-
-		UbaAgentArgs.Finalize();
+			*AgentConnectionArg,
+			"-nopoll",				// -nopoll recommended when running on remote Horde agents to make sure they exit after completion. Otherwise, it keeps running.
+			"-listenTimeout=5",		// Agent will wait 5 seconds for this thread to connect (Server_AddClient does the connect)
+			"-quiet",				// Skip all the agent logging that would be sent over to here
+			"-maxidle=15",			// After 15 seconds of idling agent will automatically disconnect
+			"-Dir=%UE_HORDE_SHARED_DIR%\\Uba",
+			"-Eventfile=%UE_HORDE_TERMINATION_SIGNAL_FILE%",
+		};
 
 		// If the machine does not run Windows, enable the compatibility layer Wine to run UbaAgent.exe on POSIX systems
 		#if PLATFORM_WINDOWS
@@ -449,16 +364,16 @@ void FUbaHordeAgentManager::ThreadAgent(FHordeAgentWrapper& Wrapper)
 			return;
 		}
 
-		Agent->Execute(AppName, UbaAgentArgs.Get().GetData(), (size_t)UbaAgentArgs.Get().Num(), nullptr, nullptr, 0, bUseWine);
+		Agent->Execute(AppName, UbaAgentArgs, UE_ARRAY_COUNT(UbaAgentArgs), nullptr, nullptr, 0, bUseWine);
 
 		// Log remote execution
 		FString UbaAgentCmdArgs = ANSI_TO_TCHAR(AppName);
-		for (const char* Arg : UbaAgentArgs.Get())
+		for (const char* Arg : UbaAgentArgs)
 		{
 			UbaAgentCmdArgs += TEXT(" ");
 			UbaAgentCmdArgs += ANSI_TO_TCHAR(Arg);
 		}
-		UE_LOG(LogUbaHorde, Log, TEXT("Remote execution on Horde machine [%s:%u]: %s"), *UbaAgentIp, UbaPort.Port, *UbaAgentCmdArgs);
+		UE_LOG(LogUbaHorde, Log, TEXT("Remote execution on Horde machine [%s:%u]: %s"), *Agent->GetMachineInfo().Ip, UbaPort, *UbaAgentCmdArgs);
 
 		MachineCoreCount = MachineInfo.LogicalCores;
 		EstimatedCoreCount += MachineCoreCount;
@@ -480,32 +395,15 @@ void FUbaHordeAgentManager::ThreadAgent(FHordeAgentWrapper& Wrapper)
 
 		if (callCounter++ == 2)
 		{
-			FString CryptoNonce16;
-			if (Agent->GetMachineInfo().Encryption != EUbaHordeEncryption::None)
-			{
-				// Generate random crypto nonce for UbaAgent
-				CryptoNonce16.Reserve(16);
-				for (int32 CharacterIndex = 0; CharacterIndex < 16; ++CharacterIndex)
-				{
-					CryptoNonce16.AppendChar(TEXT("0123456789abcdef")[FMath::Rand() % 0xF]);
-				}
-			}
-
-			// Add this machine as client to the remote agent.
-			// Use JSON field "port" when adding clients while "agentPort" is used for UbaAgent to listen to on the server side.
-			const bool bAddClientSuccess = AddClientCallback(
-				AddClientUserData,
-				StringCast<uba::tchar>(*UbaAgentIp).Get(),
-				static_cast<uint16>(UbaPort.Port),
-				Agent->GetMachineInfo().Encryption != EUbaHordeEncryption::None ? StringCast<uba::tchar>(*CryptoNonce16).Get() : nullptr
-			);
+			// Add this machine as client to the remote agent
+			const FString& IpAddress = Agent->GetMachineInfo().Ip;
+			const bool bAddClientSuccess = 	m_callback(m_userData, StringCast<uba::tchar>(*IpAddress).Get(), static_cast<uint16>(UbaPort));
 
 			if (!bAddClientSuccess)
 			{
-				UE_LOG(LogUbaHorde, Display, TEXT("Server_AddClient(%s:%u) failed"), *UbaAgentIp, UbaPort.Port);
+				UE_LOG(LogUbaHorde, Display, TEXT("Server_AddClient(%s:%u) failed"), *IpAddress, UbaPort);
 				return;
 			}
 		}
 	}
 }
-
diff --git a/Engine/Source/Developer/UbaCoordinatorHorde/Private/UbaHordeComputeTransport.cpp b/Engine/Source/Developer/UbaCoordinatorHorde/Private/UbaHordeComputeTransport.cpp
--- a/Engine/Source/Developer/UbaCoordinatorHorde/Private/UbaHordeComputeTransport.cpp
+++ b/Engine/Source/Developer/UbaCoordinatorHorde/Private/UbaHordeComputeTransport.cpp
@@ -14,12 +14,12 @@ FUbaHordeComputeTransport::FUbaHordeComputeTransport(const FHordeRemoteMachineIn
 	Socket = SocketSubsystem->CreateSocket(TEXT("Stream"), TEXT(""));
 	TSharedPtr<FInternetAddr> InternetAddr;
 
-	TSharedPtr<FInternetAddr> Address = SocketSubsystem->GetAddressFromString(MachineInfo.GetConnectionAddress());
-	Address->SetPort(MachineInfo.GetConnectionPort().Port);
+	TSharedPtr<FInternetAddr> Address = SocketSubsystem->GetAddressFromString(MachineInfo.Ip);
+	Address->SetPort(MachineInfo.Port);
 
 	if (!Socket->Connect(*Address))
 	{
-		UE_LOG(LogUbaHorde, Display, TEXT("Failed to connect to Horde compute service [%s]"), *Address->ToString(true));
+		UE_LOG(LogUbaHorde, Display, TEXT("Failed to connect to Horde Agent [%s:%d]"), *MachineInfo.Ip, (int32)MachineInfo.Port);
 
 		SocketSubsystem->DestroySocket(Socket);
 
@@ -43,12 +43,12 @@ size_t FUbaHordeComputeTransport::Send(const void* Data, size_t Size)
 	if (!Socket->Send((const uint8*)Data, Size, NumBytesSent))
 	{
 		// We can't log this since the other side could have disconnected us (causing recv to fail). This happens often in the horde setup
-		//UE_LOG(LogUbaHorde, Error, TEXT("Failed to send data to the Horde compute service: %llu %s"), (uint64)Size, Size == 1 ? TEXT("byte") : TEXT("bytes"));
+		//UE_LOG(LogUbaHorde, Error, TEXT("Failed to send data to the Horde Agent: %llu %s"), (uint64)Size, Size == 1 ? TEXT("byte") : TEXT("bytes"));
 		bHasErrors = true;
 		return 0;
 	}
 
-	UE_LOG(LogUbaHorde, Verbose, TEXT("Sent message to Horde compute service: %d %s"), NumBytesSent, NumBytesSent == 1 ? TEXT("byte") : TEXT("bytes"));
+	UE_LOG(LogUbaHorde, Verbose, TEXT("Sent message to horde agent: %d %s"), NumBytesSent, NumBytesSent == 1 ? TEXT("byte") : TEXT("bytes"));
 
 	return NumBytesSent;
 }
@@ -61,11 +61,11 @@ size_t FUbaHordeComputeTransport::Recv(void* Data, size_t Size)
 	{
 		// We can't log this since the other side could have disconnected us (causing recv to fail). This happens often in the horde setup
 		//if (!bIsClosed)
-		//	UE_LOG(LogUbaHorde, Error, TEXT("Failed to receive data from the Horde compute service: %llu %s"), (uint64)Size, Size == 1 ? TEXT("byte") : TEXT("bytes"));
+		//	UE_LOG(LogUbaHorde, Error, TEXT("Failed to receive data from the Horde Agent: %llu %s"), (uint64)Size, Size == 1 ? TEXT("byte") : TEXT("bytes"));
 		return 0;
 	}
 
-	UE_LOG(LogUbaHorde, Verbose, TEXT("Received message from Horde compute service: %d %s"), NumBytesRead, NumBytesRead == 1 ? TEXT("byte") : TEXT("bytes"));
+	UE_LOG(LogUbaHorde, Verbose, TEXT("Received message from horde agent: %d %s"), NumBytesRead, NumBytesRead == 1 ? TEXT("byte") : TEXT("bytes"));
 
 	return NumBytesRead;
 }
diff --git a/Engine/Source/Developer/UbaCoordinatorHorde/Private/UbaHordeConfig.cpp b/Engine/Source/Developer/UbaCoordinatorHorde/Private/UbaHordeConfig.cpp
deleted file mode 100644
diff --git a/Engine/Source/Developer/UbaCoordinatorHorde/Private/UbaHordeMetaClient.cpp b/Engine/Source/Developer/UbaCoordinatorHorde/Private/UbaHordeMetaClient.cpp
--- a/Engine/Source/Developer/UbaCoordinatorHorde/Private/UbaHordeMetaClient.cpp
+++ b/Engine/Source/Developer/UbaCoordinatorHorde/Private/UbaHordeMetaClient.cpp
@@ -1,13 +1,11 @@
 // Copyright Epic Games, Inc. All Rights Reserved.
 
 #include "UbaHordeMetaClient.h"
-#include "UbaHordeConfig.h"
 #include "Horde.h"
 #include "Serialization/JsonReader.h"
 #include "Serialization/JsonSerializer.h"
 #include "Misc/App.h"
 #include "Misc/ConfigCacheIni.h"
-#include "Misc/ScopeExit.h"
 #include "HttpServerConstants.h"
 
 DEFINE_LOG_CATEGORY(LogUbaHorde);
@@ -37,133 +35,56 @@ bool FUbaHordeMetaClient::RefreshHttpClient()
 	return true;
 }
 
-FString FUbaHordeMetaClient::BuildHordeRequestJsonBody(const FString& PoolId, EUbaHordeConnectionMode ConnectionMode, EUbaHordeEncryption Encryption, bool bExclusiveAccess)
+static const TCHAR* LexToString(EHordeTransportEncryption InTransportEncrypted)
 {
-	// Construct JSON body for Horde request, e.g. '{ "connection": { "modePreference": "relay", "encryption": "aes" } }'
-	const TSharedPtr<FJsonObject> RequestJNode = MakeShared<FJsonObject>();
+	switch (InTransportEncrypted)
 	{
-		if (!PoolId.IsEmpty() || bExclusiveAccess)
-		{
-			const TSharedPtr<FJsonObject> RequirementsJNode = MakeShared<FJsonObject>();
-			{
-				if (!PoolId.IsEmpty())
-				{
-					RequirementsJNode->SetStringField(TEXT("pool"), PoolId);
-				}
-				RequirementsJNode->SetBoolField(TEXT("exclusive"), bExclusiveAccess);
-			}
-			RequestJNode->SetObjectField(TEXT("requirements"), RequirementsJNode);
-		}
-
-		const TSharedPtr<FJsonObject> ConnectionJNode = MakeShared<FJsonObject>();
-		{
-			ConnectionJNode->SetStringField(TEXT("modePreference"), LexToString(ConnectionMode));
-			if (Encryption != EUbaHordeEncryption::None)
-			{
-				ConnectionJNode->SetStringField(TEXT("encryption"), LexToString(Encryption));
-			}
-
-			const TSharedPtr<FJsonObject> PortsJNode = MakeShared<FJsonObject>();
-			{
-				PortsJNode->SetNumberField(TEXT("UbaPort"), (double)FHordeRemoteMachineInfo::UbaPort);
-				PortsJNode->SetNumberField(TEXT("UbaProxyPort"), (double)FHordeRemoteMachineInfo::UbaProxyPort);
-			}
-			ConnectionJNode->SetObjectField(TEXT("ports"), PortsJNode);
-		}
-		RequestJNode->SetObjectField(TEXT("connection"), ConnectionJNode);
+	case EHordeTransportEncryption::None:	break;
+	case EHordeTransportEncryption::AES:	return TEXT("aes");
 	}
-
-	FString JsonContent;
-	const TSharedRef<TJsonWriter<>> JsonWriter = TJsonWriterFactory<>::Create(&JsonContent);
-	FJsonSerializer::Serialize(RequestJNode.ToSharedRef(), JsonWriter);
-	return JsonContent;
+	return TEXT("");
 }
 
-TSharedPtr<FUbaHordeMetaClient::HordeClusterPromise, ESPMode::ThreadSafe> FUbaHordeMetaClient::RequestClusterId(const FString& HordeRequestJsonBody)
+static void BuildHordeRequestJsonBody(TStringBuilder<1024>& RequestJsonBody, const FString& InPoolId, EHordeTransportEncryption InTransportEncrypted)
 {
-	TSharedPtr<HordeClusterPromise> Promise = MakeShared<HordeClusterPromise, ESPMode::ThreadSafe>();
-
-	FHttpRequestRef Request = HttpClient->CreateRequest(TEXT("POST"), TEXT("api/v2/compute/_cluster"));
-
-	Request->SetDelegateThreadPolicy(EHttpRequestDelegateThreadPolicy::CompleteOnHttpThread);
-	Request->SetContentAsString(HordeRequestJsonBody);
-	Request->SetHeader(TEXT("Content-Type"), TEXT("application/json"));
-
-	Request->OnProcessRequestComplete().BindLambda(
-		[this, Promise](FHttpRequestPtr /*Request*/, FHttpResponsePtr HttpResponse, bool bConnectedSuccessfully)
-		{
-			FHordeClusterInfo Info;
-
-			ON_SCOPE_EXIT
-			{
-				Promise->SetValue(MakeTuple(HttpResponse, Info));
-			};
-
-			if (!bConnectedSuccessfully || !HttpResponse.IsValid())
-			{
-				UE_LOG(LogUbaHorde, Display, TEXT("No response from Horde"));
-				return;
-			}
+	RequestJsonBody.AppendChar(TEXT('{'));
 
-			const EHttpServerResponseCodes ResponseCode = (EHttpServerResponseCodes)HttpResponse->GetResponseCode();
-			const FString ResponseStr = HttpResponse->GetContentAsString();
-
-			if (ResponseCode == EHttpServerResponseCodes::ServiceUnavail ||
-				ResponseCode == EHttpServerResponseCodes::TooManyRequests)
-			{
-				// Service Unavailable
-				UE_LOG(LogUbaHorde, Display, TEXT("Horde agent request returned with HTTP/%d: %s"), ResponseCode, *ResponseStr);
-				return;
-			}
-
-			TSharedPtr<FJsonValue> OutJson;
-			TSharedRef<TJsonReader<>> Reader = TJsonReaderFactory<>::Create(ResponseStr);
-
-			if (!FJsonSerializer::Deserialize(Reader, OutJson, FJsonSerializer::EFlags::None))
-			{
-				// Report invalid response body with Display verbosity only, since this should not fail a CIS job
-				UE_LOG(LogUbaHorde, Display, TEXT("Invalid response body: %s"), *ResponseStr);
-				return;
-			}
-
-			TSharedPtr<FJsonValue> ClusterIdValue = OutJson->AsObject()->TryGetField(TEXT("clusterId"));
-			if (!ClusterIdValue.Get())
-			{
-				// Report invalid response body with Display verbosity only, since this should not fail a CIS job
-				UE_LOG(LogUbaHorde, Display, TEXT("Invalid response body: %s"), *ResponseStr);
-				return;
-			}
-
-			// Successfully return resolved cluster ID
-			Info.ClusterId = ClusterIdValue->AsString();
-		});
+	// Append requirements: pool ID and whether it should be allocated exclusively for this request
+	RequestJsonBody.Appendf(TEXT("\"requirements\":{\"pool\":\"%s\",\"exclusive\":true}"), *InPoolId);
 
-	Request->ProcessRequest();
+	// Append connection information if encryption is enabled
+	if (InTransportEncrypted != EHordeTransportEncryption::None)
+	{
+		RequestJsonBody.Appendf(TEXT(",\"connection\":{\"encryption\":\"%s\",\"modePreference\":\"Direct\"}"), LexToString(InTransportEncrypted));
+	}
 
-	return Promise;
+	RequestJsonBody.AppendChar(TEXT('}'));
 }
 
-TSharedPtr<FUbaHordeMetaClient::HordeMachinePromise, ESPMode::ThreadSafe> FUbaHordeMetaClient::RequestMachine(const FString& HordeRequestJsonBody, const TCHAR* ClusterId)
+TSharedPtr<FUbaHordeMetaClient::HordeMachinePromise, ESPMode::ThreadSafe> FUbaHordeMetaClient::RequestMachine(const FString& PoolId, const TCHAR* Machine, EHordeTransportEncryption Encryption)
 {
 	TSharedPtr<HordeMachinePromise> Promise = MakeShared<HordeMachinePromise, ESPMode::ThreadSafe>();
 
-	const FString ResourcePath = FString::Format(TEXT("api/v2/compute/{0}"), { ClusterId != nullptr && *ClusterId != TEXT('\0') ? ClusterId : TEXT("default")});
+	const FString ResourcePath = FString::Format(TEXT("api/v2/compute/{0}"), { Machine != nullptr ? Machine : TEXT("default") });
+
 	FHttpRequestRef Request = HttpClient->CreateRequest(TEXT("POST"), *ResourcePath);
 
 	Request->SetDelegateThreadPolicy(EHttpRequestDelegateThreadPolicy::CompleteOnHttpThread);
-	Request->SetContentAsString(HordeRequestJsonBody);
+
+	// Post JSON document with constraints to acquire a Horde agent. Use pool of agents (e.g. "BoxLinux" or "BoxWin")
+	// and require exclusive access or UbaStorage will fail to initialize the next time we connect to the same machine.
+	TStringBuilder<1024> RequestJsonBody;
+	BuildHordeRequestJsonBody(RequestJsonBody, PoolId, Encryption);
+
+	Request->SetContentAsString(RequestJsonBody.ToString());
 	Request->SetHeader(TEXT("Content-Type"), TEXT("application/json"));
 
+	UE_LOG(LogUbaHorde, Log, TEXT("Requesting Horde agent with JSON descriptor: '%s'"), RequestJsonBody.ToString());
+
 	Request->OnProcessRequestComplete().BindLambda(
-		[this, Promise, ClusterId = FString(ClusterId)](FHttpRequestPtr /*Request*/, FHttpResponsePtr HttpResponse, bool bConnectedSuccessfully)
+		[this, Promise, Encryption](FHttpRequestPtr /*Request*/, FHttpResponsePtr HttpResponse, bool bConnectedSuccessfully)
 		{
 			FHordeRemoteMachineInfo Info;
-
-			ON_SCOPE_EXIT
-			{
-				Promise->SetValue(MakeTuple(HttpResponse, Info));
-			};
-
 			Info.Ip = TEXT("");
 			Info.Port = 0xFFFF;
 			Info.bRunsWindowOS = false;
@@ -172,6 +93,7 @@ TSharedPtr<FUbaHordeMetaClient::HordeMachinePromise, ESPMode::ThreadSafe> FUbaHo
 			if (!bConnectedSuccessfully || !HttpResponse.IsValid())
 			{
 				UE_LOG(LogUbaHorde, Display, TEXT("No response from Horde"));
+				Promise->SetValue(MakeTuple(HttpResponse, Info));
 				return;
 			}
 
@@ -183,6 +105,7 @@ TSharedPtr<FUbaHordeMetaClient::HordeMachinePromise, ESPMode::ThreadSafe> FUbaHo
 			{
 				// Service Unavailable
 				UE_LOG(LogUbaHorde, Display, TEXT("Horde agent request returned with HTTP/%d: %s"), ResponseCode, *ResponseStr);
+				Promise->SetValue(MakeTuple(HttpResponse, Info));
 				return;
 			}
 
@@ -193,6 +116,7 @@ TSharedPtr<FUbaHordeMetaClient::HordeMachinePromise, ESPMode::ThreadSafe> FUbaHo
 			{
 				// Report invalid response body with Display verbosity only, since this should not fail a CIS job
 				UE_LOG(LogUbaHorde, Display, TEXT("Invalid response body: %s"), *ResponseStr);
+				Promise->SetValue(MakeTuple(HttpResponse, Info));
 				return;
 			}
 
@@ -204,38 +128,10 @@ TSharedPtr<FUbaHordeMetaClient::HordeMachinePromise, ESPMode::ThreadSafe> FUbaHo
 			{
 				// Report invalid response body with Display verbosity only, since this should not fail a CIS job
 				UE_LOG(LogUbaHorde, Display, TEXT("Invalid response body: %s"), *ResponseStr);
+				Promise->SetValue(MakeTuple(HttpResponse, Info));
 				return;
 			}
 
-			// Check for optional port mapping array
-			if (TSharedPtr<FJsonValue> PortsValue = OutJson->AsObject()->TryGetField(TEXT("ports")))
-			{
-				for (auto Iter = PortsValue->AsObject()->Values.CreateConstIterator(); Iter; ++Iter)
-				{
-					FHordeRemoteMachineInfo::FPortInfo& PortInfo = Info.Ports.Add(Iter.Key());
-					if (TSharedPtr<FJsonValue> PortInfoPortValue = Iter.Value()->AsObject()->TryGetField(TEXT("port")))
-					{
-						PortInfo.Port = (uint16)PortInfoPortValue->AsNumber();
-					}
-					if (TSharedPtr<FJsonValue> PortInfoAgentPortValue = Iter.Value()->AsObject()->TryGetField(TEXT("agentPort")))
-					{
-						PortInfo.AgentPort = (uint16)PortInfoAgentPortValue->AsNumber();
-					}
-				}
-			}
-
-			// Check for connection mode and address - "connectionAddress" field is meaningless without a mode other than the default one
-			if (TSharedPtr<FJsonValue> ConnectionModeValue = OutJson->AsObject()->TryGetField(TEXT("connectionMode")))
-			{
-				if (LexFromString(Info.ConnectionMode, *ConnectionModeValue->AsString()))
-				{
-					if (TSharedPtr<FJsonValue> ConnectionAddressValue = OutJson->AsObject()->TryGetField(TEXT("connectionAddress")))
-					{
-						Info.ConnectionAddress = ConnectionAddressValue->AsString();
-					}
-				}
-			}
-
 			FString OsFamily(TEXT("UNKNOWN-OS"));
 
 			uint16 LogicalCores = 0;
@@ -284,38 +180,29 @@ TSharedPtr<FUbaHordeMetaClient::HordeMachinePromise, ESPMode::ThreadSafe> FUbaHo
 			FString IpString = IpValue->AsString();
 			uint16 PortNumber = (uint16)PortValue->AsNumber();
 
-			// Return final response information
-			Info.Ip = IpString;
-			Info.Port = PortNumber;
-			FString::ToHexBlob(NonceString, Info.Nonce, HORDE_NONCE_SIZE);
-
-			FString HordeMachineSummary = FString::Printf(TEXT("UBA Horde machine assigned (%s) on '%s' [%s:%u]"), *OsFamily, *ClusterId, *Info.GetConnectionAddress(), (uint32)Info.GetConnectionPort().Port);
+			FString HordeMachineSummary = FString::Printf(TEXT("UBA Horde machine assigned (%s) [%s:%u]"), *OsFamily, *IpString, (uint32)PortNumber);
 
 #if WITH_SSL
-			if (TSharedPtr<FJsonValue> EncryptionValue = OutJson->AsObject()->TryGetField(TEXT("encryption")))
+			if (Encryption == EHordeTransportEncryption::AES)
 			{
-				if (LexFromString(Info.Encryption, *EncryptionValue->AsString()))
+				TSharedPtr<FJsonValue> KeyValue = OutJson->AsObject()->TryGetField(TEXT("key"));
+				FString KeyString = KeyValue ? KeyValue->AsString() : FString();
+				if (KeyString.IsEmpty())
 				{
-					if (Info.Encryption == EUbaHordeEncryption::AES)
-					{
-						TSharedPtr<FJsonValue> KeyValue = OutJson->AsObject()->TryGetField(TEXT("key"));
-						FString KeyString = KeyValue ? KeyValue->AsString() : FString();
-						if (KeyString.IsEmpty())
-						{
-							HordeMachineSummary += TEXT(" [AES key missing]");
-						}
-						else if (KeyString.Len() != sizeof(Info.Key) * 2)
-						{
-							HordeMachineSummary += TEXT(" [AES key corrupted]");
-						}
-						else
-						{
-							FString::ToHexBlob(KeyString, Info.Key, sizeof(Info.Key));
-							HordeMachineSummary += TEXT(" [AES key received]");
-						}
-					}
+					HordeMachineSummary += TEXT(" [AES key missing]");
+				}
+				else if (KeyString.Len() != sizeof(Info.Key) * 2)
+				{
+					HordeMachineSummary += TEXT(" [AES key corrupted]");
+				}
+				else
+				{
+					FString::ToHexBlob(KeyString, Info.Key, sizeof(Info.Key));
+					HordeMachineSummary += TEXT(" [AES key received]");
 				}
 			}
+
+			Info.Encryption = Encryption;
 #endif
 
 			if (TSharedPtr<FJsonValue> LeaseIdValue = OutJson->AsObject()->TryGetField(TEXT("leaseId")))
@@ -324,6 +211,13 @@ TSharedPtr<FUbaHordeMetaClient::HordeMachinePromise, ESPMode::ThreadSafe> FUbaHo
 			}
 
 			UE_LOG(LogUbaHorde, Display, TEXT("%s"), *HordeMachineSummary);
+
+			// Return final response information
+			Info.Ip = IpString;
+			Info.Port = PortNumber;
+			FString::ToHexBlob(NonceString, Info.Nonce, HORDE_NONCE_SIZE);
+
+			Promise->SetValue(MakeTuple(HttpResponse, Info));
 		});
 
 	Request->ProcessRequest();
diff --git a/Engine/Source/Developer/UbaCoordinatorHorde/Private/UbaHordeMetaClient.h b/Engine/Source/Developer/UbaCoordinatorHorde/Private/UbaHordeMetaClient.h
--- a/Engine/Source/Developer/UbaCoordinatorHorde/Private/UbaHordeMetaClient.h
+++ b/Engine/Source/Developer/UbaCoordinatorHorde/Private/UbaHordeMetaClient.h
@@ -8,7 +8,6 @@
 #include "Interfaces/IHttpResponse.h"
 #include "HordeHttpClient.h"
 #include "UbaHordeAgentManager.h"
-#include "UbaHordeConfig.h"
 
 DECLARE_LOG_CATEGORY_EXTERN(LogUbaHorde, Log, Log);
 
@@ -23,49 +22,14 @@ DECLARE_LOG_CATEGORY_EXTERN(LogUbaHorde, Log, Log);
 // - Nonce will be all 0
 struct FHordeRemoteMachineInfo
 {
-	struct FPortInfo
-	{
-		FPortInfo() = default;
-		FPortInfo(uint16 InPort, uint16 InAgentPort)
-			: Port(InPort)
-			, AgentPort(InAgentPort)
-		{
-		}
-
-		uint16 Port = 0;
-		uint16 AgentPort = 0;
-	};
-
 	FString Ip;
-	EUbaHordeConnectionMode ConnectionMode = EUbaHordeConnectionMode::Direct;
-	FString ConnectionAddress; // Onyl used for relay mode
 	uint16 Port = 0;
-	TMap<FString, FPortInfo> Ports; // Maps compute service names to port/agent-port pairs, e.g. "_horde_compute" -> {Port: 12345, AgentPort: 7000}
 	uint16 LogicalCores = 0;
-	EUbaHordeEncryption Encryption = EUbaHordeEncryption::None;
+	EHordeTransportEncryption Encryption = EHordeTransportEncryption::None;
 	uint8 Nonce[HORDE_NONCE_SIZE] = {};
 	uint8 Key[HORDE_KEY_SIZE] = {};
 	TArray<uint8> CertificateData; // X.509 certificate for SSL connections in Pkcs12 format.
 	bool bRunsWindowOS = false;
-
-	static constexpr uint16 UbaPort = 7001;
-	static constexpr uint16 UbaProxyPort = 7002;
-
-	const FString& GetConnectionAddress() const
-	{
-		return ConnectionMode == EUbaHordeConnectionMode::Relay ? ConnectionAddress : Ip;
-	}
-
-	FPortInfo GetConnectionPort() const
-	{
-		static const FString HordeComputeIdentifier = TEXT("_horde_compute");
-		return ConnectionMode == EUbaHordeConnectionMode::Relay ? Ports[HordeComputeIdentifier] : FPortInfo(Port, Port);
-	}
-};
-
-struct FHordeClusterInfo
-{
-	FString ClusterId = TEXT("default");
 };
 
 // This encapsulates the mechanism of talking to the Horde "meta server". The server which
@@ -75,16 +39,11 @@ class FUbaHordeMetaClient final
 public:
 	// We return the HttpResponse in case more information could be used out of it later.
 	using HordeMachinePromise = TPromise<TTuple<FHttpResponsePtr, FHordeRemoteMachineInfo>>;
-	using HordeClusterPromise = TPromise<TTuple<FHttpResponsePtr, FHordeClusterInfo>>;
-
+	
 	FUbaHordeMetaClient() = default;
 
 	bool RefreshHttpClient();
 
-	static FString BuildHordeRequestJsonBody(const FString& PoolId, EUbaHordeConnectionMode ConnectionMode, EUbaHordeEncryption Encryption, bool bExclusiveAccess = true);
-
-	TSharedPtr<HordeClusterPromise, ESPMode::ThreadSafe> RequestClusterId(const FString& HordeRequestJsonBody);
-
 	// This will make a request to Horde for a remote machine to do work on.
 	// Example of actually getting the FHordeRemoteMachineInfo struct:
 	// ```
@@ -93,9 +52,11 @@ public:
 	// ... // Can to asynchronous work
 	// auto Future.Wait(); // Wait for the response to arrive
 	// FHordeRemoteMachineInfo MachineInfo = Future.Get().Value;
-	TSharedPtr<HordeMachinePromise, ESPMode::ThreadSafe> RequestMachine(const FString& HordeRequestJsonBody, const TCHAR* ClusterId = nullptr);
+	TSharedPtr<HordeMachinePromise, ESPMode::ThreadSafe> RequestMachine(const FString& PoolId, const TCHAR* Machine = nullptr, EHordeTransportEncryption Encryption = EHordeTransportEncryption::None);
 
 private:
+	void ParseConfig(const FString& IniFilename, const FString& IniSection);
+
 	FString ServerUrl;
 	TUniquePtr<FHordeHttpClient> HttpClient;
 };
diff --git a/Engine/Source/Developer/UbaCoordinatorHorde/Public/UbaHordeAgentManager.h b/Engine/Source/Developer/UbaCoordinatorHorde/Public/UbaHordeAgentManager.h
--- a/Engine/Source/Developer/UbaCoordinatorHorde/Public/UbaHordeAgentManager.h
+++ b/Engine/Source/Developer/UbaCoordinatorHorde/Public/UbaHordeAgentManager.h
@@ -5,22 +5,33 @@
 #include "Containers/Array.h"
 #include "HAL/Thread.h"
 #include "UbaBase.h"
-#include "UbaHordeConfig.h"
 
 class FEvent;
 class FUbaHordeMetaClient;
 namespace uba { class NetworkServer; }
 
+enum class EHordeTransportEncryption
+{
+	None,
+	AES,
+};
+
 class FUbaHordeAgentManager
 {
 public:
 	UBACOORDINATORHORDE_API FUbaHordeAgentManager(const FString& InWorkingDir, const FString& BinariesPath);
 	UBACOORDINATORHORDE_API ~FUbaHordeAgentManager();
 
+	void SetPool(const FString& InPool) { Pool = InPool; }
+	void SetMaxCoreCount(uint32 Count) { MaxCores = Count; }
+	void SetEncryption(EHordeTransportEncryption InEncryption) { Encryption = InEncryption; }
+	void SetUbaHost(const FAnsiString& InHost) { UbaHost = InHost; }
+	void SetUbaPort(uint32 InPort) { UbaPort = InPort; }
+
 	UBACOORDINATORHORDE_API void SetTargetCoreCount(uint32 Count);
 
-	using FAddClientCallback = bool(void* UserData, const uba::tchar* Ip, uint16 Port, const uba::tchar* Crypto16Characters);
-	UBACOORDINATORHORDE_API void SetAddClientCallback(FAddClientCallback* Callback, void* UserData);
+	using AddClientCallback = bool(void* userData, const uba::tchar* ip, uint16 port);
+	UBACOORDINATORHORDE_API void SetAddClientCallback(AddClientCallback* callback, void* userData);
 
 	// Returns the number of agents currently handled by this agent manager.
 	UBACOORDINATORHORDE_API int32 GetAgentCount() const;
@@ -41,6 +52,12 @@ private:
 	FString WorkingDir;
 	FString BinariesPath;
 
+	FString Pool;
+	FAnsiString UbaHost;
+	uint32 UbaPort = 7001;
+	uint32 MaxCores = 500; // Assume 500 by default in case this configuration is missing but a UBA pool was found
+	EHordeTransportEncryption Encryption = EHordeTransportEncryption::None;
+
 	TUniquePtr<FUbaHordeMetaClient> HordeMetaClient;
 
 	FCriticalSection BundleRefPathsLock;
@@ -55,8 +72,8 @@ private:
 	TAtomic<uint32> ActiveCoreCount;
 	TAtomic<bool> bAskForAgents;
 
-	FAddClientCallback* AddClientCallback = nullptr;
-	void* AddClientUserData = nullptr;
+	AddClientCallback* m_callback = nullptr;
+	void* m_userData = nullptr;
 
 	FUbaHordeAgentManager(const FUbaHordeAgentManager&) = delete;
 	void operator=(const FUbaHordeAgentManager&) = delete;
diff --git a/Engine/Source/Developer/UbaCoordinatorHorde/Public/UbaHordeConfig.h b/Engine/Source/Developer/UbaCoordinatorHorde/Public/UbaHordeConfig.h
deleted file mode 100644
diff --git a/Engine/Source/Developer/UbaCoordinatorHorde/UbaCoordinatorHorde.Build.cs b/Engine/Source/Developer/UbaCoordinatorHorde/UbaCoordinatorHorde.Build.cs
--- a/Engine/Source/Developer/UbaCoordinatorHorde/UbaCoordinatorHorde.Build.cs
+++ b/Engine/Source/Developer/UbaCoordinatorHorde/UbaCoordinatorHorde.Build.cs
@@ -22,7 +22,6 @@ public class UbaCoordinatorHorde : ModuleRules
 			"Sockets",
 			"SSL",
 			"OpenSSL",
-			"XmlParser",
 		});
 
 		PrivateIncludePathModuleNames.AddRange(new string[] {
