[UBA] Align configuration of maximum local UBA workers in UbaController with UBT while also honoring legacy options to limit number of local workers.

This changes include two updates to the shader compiler configuration:
1. Use the same XML and command line option to set the number of maxium local workers (i.e. "MaxParallelActions") in the UbaController for shader compilation as for UBT.
2. Also apply local worker limitation set by the legacy options from the shader compiling manager read from INI files:
[DevOptions.Shaders]:NumUnusedShaderCompilingThreads, NumUnusedShaderCompilingThreadsDuringGame, ShaderCompilerCoreCountThreshold.

diff --git a/Engine/Plugins/UbaController/Source/Private/UbaControllerModule.cpp b/Engine/Plugins/UbaController/Source/Private/UbaControllerModule.cpp
--- a/Engine/Plugins/UbaController/Source/Private/UbaControllerModule.cpp
+++ b/Engine/Plugins/UbaController/Source/Private/UbaControllerModule.cpp
@@ -267,7 +267,7 @@ void FUbaControllerModule::InitializeController()
 bool FUbaControllerModule::SupportsLocalWorkers()
 {
 	// UbaController supports local workers if the maximum number of local cores is greater than zero or -1 (special value to use all available cores)
-	const int32 MaxLocalCores = FUbaHordeConfig::Get().HordeMaxLocalCores;
+	const int32 MaxLocalCores = FUbaHordeConfig::Get().MaxParallelActions;
 	return MaxLocalCores > 0 || MaxLocalCores == -1;
 }
 
diff --git a/Engine/Plugins/UbaController/Source/Private/UbaJobProcessor.cpp b/Engine/Plugins/UbaController/Source/Private/UbaJobProcessor.cpp
--- a/Engine/Plugins/UbaController/Source/Private/UbaJobProcessor.cpp
+++ b/Engine/Plugins/UbaController/Source/Private/UbaJobProcessor.cpp
@@ -155,11 +155,7 @@ FUbaJobProcessor::FUbaJobProcessor(
 			checkf(false, TEXT("%s"), UBASTRING_TO_TCHAR(text));
 		});
 
-	MaxLocalParallelJobs = FUbaHordeConfig::Get().HordeMaxLocalCores;
-	if (MaxLocalParallelJobs == -1)
-	{
-		MaxLocalParallelJobs = FPlatformMisc::NumberOfCoresIncludingHyperthreads();
-	}
+	UpdateMaxLocalParallelJobs();
 }
 
 FUbaJobProcessor::~FUbaJobProcessor()
@@ -167,6 +163,23 @@ FUbaJobProcessor::~FUbaJobProcessor()
 	delete Thread;
 }
 
+void FUbaJobProcessor::UpdateMaxLocalParallelJobs()
+{
+	// Limit number of parallel jobs by UBA/Horde configuration
+	MaxLocalParallelJobs = FUbaHordeConfig::Get().MaxParallelActions;
+	if (MaxLocalParallelJobs == -1)
+	{
+		MaxLocalParallelJobs = FPlatformMisc::NumberOfCoresIncludingHyperthreads();
+	}
+
+	// Also apply limits from shader compiling manager, i.e. [DevOptions.Shaders]:NumUnusedShaderCompilingThreads etc.
+	const int32 MaxNumLocalWorkers = ControllerModule.GetMaxNumLocalWorkers();
+	if (MaxNumLocalWorkers != -1)
+	{
+		MaxLocalParallelJobs = FMath::Min(MaxLocalParallelJobs, MaxNumLocalWorkers);
+	}
+}
+
 void FUbaJobProcessor::CalculateKnownInputs()
 {
 	// TODO: This is ShaderCompileWorker specific and this code is designed to handle all kinds of distributed workload.
@@ -358,6 +371,7 @@ void FUbaJobProcessor::StartUba()
 	UbaSessionServer = SessionServer_Create(info);
 
 	CalculateKnownInputs();
+	UpdateMaxLocalParallelJobs();
 
 	UbaScheduler = Scheduler_Create(UbaSessionServer, MaxLocalParallelJobs, UbaJobProcessorOptions::bAllowProcessReuse);
 	Scheduler_Start(UbaScheduler);
diff --git a/Engine/Plugins/UbaController/Source/Public/UbaControllerModule.h b/Engine/Plugins/UbaController/Source/Public/UbaControllerModule.h
--- a/Engine/Plugins/UbaController/Source/Public/UbaControllerModule.h
+++ b/Engine/Plugins/UbaController/Source/Public/UbaControllerModule.h
@@ -27,7 +27,7 @@ public:
 	virtual bool SupportsLocalWorkers() override final;
 	virtual FString CreateUniqueFilePath() override final;
 	virtual TFuture<FDistributedBuildTaskResult> EnqueueTask(const FTaskCommandData& CommandData) override final;
-	virtual bool PollStats(FDistributedBuildStats& OutStats) override;
+	virtual bool PollStats(FDistributedBuildStats& OutStats) override final;
 
 	void ReportJobProcessed(const FTaskResponse& InTaskResponse, FDistributedBuildTask* CompileTask);
 	void CleanWorkingDirectory() const;
@@ -41,6 +41,16 @@ public:
 	// Queue of tasks submitted by the engine, but not yet dispatched to the controller.
 	TQueue<FDistributedBuildTask*, EQueueMode::Mpsc> PendingRequestedCompilationTasks;
 
+	virtual void SetMaxLocalWorkers(int32 InMaxNumLocalWorkers) override final
+	{
+		MaxNumLocalWorkers = InMaxNumLocalWorkers;
+	}
+
+	inline int32 GetMaxNumLocalWorkers() const
+	{
+		return MaxNumLocalWorkers;
+	}
+
 	static FString GetTempDir();
 
 private:
@@ -56,6 +66,7 @@ private:
 
 	TAtomic<int32> NextFileID;
 	TAtomic<int32> NextTaskID;
+	TAtomic<int32> MaxNumLocalWorkers = -1;
 
 	TSharedPtr<FUbaJobProcessor> JobDispatcherThread;
 };
diff --git a/Engine/Plugins/UbaController/Source/Public/UbaJobProcessor.h b/Engine/Plugins/UbaController/Source/Public/UbaJobProcessor.h
--- a/Engine/Plugins/UbaController/Source/Public/UbaJobProcessor.h
+++ b/Engine/Plugins/UbaController/Source/Public/UbaJobProcessor.h
@@ -42,6 +42,7 @@ public:
 
 private:
 
+	void UpdateMaxLocalParallelJobs();
 	void CalculateKnownInputs();
 	void RunTaskWithUba(FDistributedBuildTask* Task);
 
diff --git a/Engine/Source/Developer/DistributedBuildInterface/Public/DistributedBuildControllerInterface.h b/Engine/Source/Developer/DistributedBuildInterface/Public/DistributedBuildControllerInterface.h
--- a/Engine/Source/Developer/DistributedBuildInterface/Public/DistributedBuildControllerInterface.h
+++ b/Engine/Source/Developer/DistributedBuildInterface/Public/DistributedBuildControllerInterface.h
@@ -81,6 +81,9 @@ public:
 	
 	virtual bool RequiresRelativePaths() { return false; }
 
+	// Sets the maxmium number of local workers. Ignored if this controller does not support local workers.
+	virtual void SetMaxLocalWorkers(int32 InMaxNumLocalWorkers) { /*dummy*/ }
+
 	virtual void InitializeController() = 0;
 	
 	// Returns true if the controller may be used.
diff --git a/Engine/Source/Developer/UbaCoordinatorHorde/Private/UbaHordeConfig.cpp b/Engine/Source/Developer/UbaCoordinatorHorde/Private/UbaHordeConfig.cpp
--- a/Engine/Source/Developer/UbaCoordinatorHorde/Private/UbaHordeConfig.cpp
+++ b/Engine/Source/Developer/UbaCoordinatorHorde/Private/UbaHordeConfig.cpp
@@ -40,7 +40,7 @@ UBACOORDINATORHORDE_API const FString& FUbaHordeConfig::GetHordePool() const
 	return DefaultHordePool;
 }
 
-static bool ParseUbaHordeConfigFromXmlFile(FUbaHordeConfig& OutConfig, const TCHAR* InBuildConfigFilename)
+bool FUbaHordeConfig::ParseConfigFromXmlFile(const TCHAR* InBuildConfigFilename)
 {
 	FXmlFile XmlFile;
 	if (!XmlFile.LoadFile(InBuildConfigFilename))
@@ -50,6 +50,8 @@ static bool ParseUbaHordeConfigFromXmlFile(FUbaHordeConfig& OutConfig, const TCH
 
 	static const FString ConfigurationTag(TEXT("Configuration"));
 	static const FString BuildConfigurationTag(TEXT("BuildConfiguration"));
+	static const FString AllowUBAExecutorTag(TEXT("bAllowUBAExecutor"));
+	static const FString MaxParallelActionsTag(TEXT("MaxParallelActions"));
 	static const FString HordeTag(TEXT("Horde"));
 	static const FString UbaDisabledTag(TEXT("bDisableHorde"));
 	static const FString ServerTag(TEXT("Server"));
@@ -66,80 +68,92 @@ static bool ParseUbaHordeConfigFromXmlFile(FUbaHordeConfig& OutConfig, const TCH
 		return false;
 	}
 
+	if (const FXmlNode* BuildConfigurationNode = ConfigurationNode->FindChildNode(BuildConfigurationTag))
+	{
+		/*if (const FXmlNode* AllowUBAExecutorNode = BuildConfigurationNode->FindChildNode(AllowUBAExecutorTag))
+		{
+			[REDACTED_EMAIL]: Once shader compiler and UBT are menat to use the same configuration source to enable UBA, make use of this XML tag
+		}*/
+		if (const FXmlNode* MaxParallelActionsNode = BuildConfigurationNode->FindChildNode(MaxParallelActionsTag))
+		{
+			MaxParallelActions = FCString::Atoi(*MaxParallelActionsNode->GetContent());
+		}
+	}
+
 	if (const FXmlNode* HordeNode = ConfigurationNode->FindChildNode(HordeTag))
 	{
 		if (const FXmlNode* UbaDisabledNode = HordeNode->FindChildNode(UbaDisabledTag))
 		{
 			if (UbaDisabledNode->GetContent() == TEXT("true"))
 			{
-				OutConfig.HordeEnabled = EUbaHordeEnabledMode::Disabled;
+				HordeEnabled = EUbaHordeEnabledMode::Disabled;
 			}
 		}
 		if (const FXmlNode* DefaultPoolNode = HordeNode->FindChildNode(DefaultPoolTag))
 		{
-			OutConfig.DefaultHordePool = DefaultPoolNode->GetContent();
+			DefaultHordePool = DefaultPoolNode->GetContent();
 		}
 		if (const FXmlNode* MacPoolNode = HordeNode->FindChildNode(MacPoolTag))
 		{
-			OutConfig.MacHordePool = MacPoolNode->GetContent();
+			MacHordePool = MacPoolNode->GetContent();
 		}
 		if (const FXmlNode* ServerNode = HordeNode->FindChildNode(ServerTag))
 		{
-			OutConfig.HordeServer = ServerNode->GetContent();
+			HordeServer = ServerNode->GetContent();
 		}
 		if (const FXmlNode* ConnectionModeNode = HordeNode->FindChildNode(ConnectionModeTag))
 		{
-			LexFromString(OutConfig.HordeConnectionMode, *ConnectionModeNode->GetContent());
+			LexFromString(HordeConnectionMode, *ConnectionModeNode->GetContent());
 		}
 		if (const FXmlNode* EncryptionNode = HordeNode->FindChildNode(EncryptionTag))
 		{
-			LexFromString(OutConfig.HordeEncryption, *EncryptionNode->GetContent());
+			LexFromString(HordeEncryption, *EncryptionNode->GetContent());
 		}
 		if (const FXmlNode* ClusterNode = HordeNode->FindChildNode(ClusterTag))
 		{
-			OutConfig.HordeCluster = ClusterNode->GetContent();
+			HordeCluster = ClusterNode->GetContent();
 		}
 		if (const FXmlNode* UBASentryUrlNode = HordeNode->FindChildNode(UBASentryUrlTag))
 		{
-			OutConfig.UbaSentryUrl = UBASentryUrlNode->GetContent();
+			UbaSentryUrl = UBASentryUrlNode->GetContent();
 		}
 	}
 
 	return true;
 }
 
-static void ParseUbaHordeConfigFromIniFile(FUbaHordeConfig& OutConfig)
+void FUbaHordeConfig::ParseConfigFromIniFile()
 {
-	GConfig->GetString(TEXT("Horde"), TEXT("ServerUrl"), OutConfig.HordeServer, GEngineIni);
-	GConfig->GetString(TEXT("Horde"), TEXT("UbaToken"), OutConfig.HordeToken, GEngineIni);
-	GConfig->GetString(TEXT("Horde"), TEXT("UbaPool"), OutConfig.DefaultHordePool, GEngineIni);
-	GConfig->GetString(TEXT("Horde"), TEXT("UbaLinuxPool"), OutConfig.LinuxHordePool, GEngineIni);
-	GConfig->GetString(TEXT("Horde"), TEXT("UbaMacPool"), OutConfig.MacHordePool, GEngineIni);
-	GConfig->GetString(TEXT("Horde"), TEXT("UbaWindowsPool"), OutConfig.WindowsHordePool, GEngineIni);
-	GConfig->GetString(TEXT("Horde"), TEXT("UbaCluster"), OutConfig.HordeCluster, GEngineIni);
-	GConfig->GetString(TEXT("Horde"), TEXT("UbaLocalHost"), OutConfig.HordeHost, GEngineIni);
-	GConfig->GetInt(TEXT("Horde"), TEXT("MaxCores"), OutConfig.HordeMaxCores, GEngineIni);
-	GConfig->GetBool(TEXT("Horde"), TEXT("UbaAllowWine"), OutConfig.bHordeAllowWine, GEngineIni);
+	GConfig->GetString(TEXT("Horde"), TEXT("ServerUrl"), HordeServer, GEngineIni);
+	GConfig->GetString(TEXT("Horde"), TEXT("UbaToken"), HordeToken, GEngineIni);
+	GConfig->GetString(TEXT("Horde"), TEXT("UbaPool"), DefaultHordePool, GEngineIni);
+	GConfig->GetString(TEXT("Horde"), TEXT("UbaLinuxPool"), LinuxHordePool, GEngineIni);
+	GConfig->GetString(TEXT("Horde"), TEXT("UbaMacPool"), MacHordePool, GEngineIni);
+	GConfig->GetString(TEXT("Horde"), TEXT("UbaWindowsPool"), WindowsHordePool, GEngineIni);
+	GConfig->GetString(TEXT("Horde"), TEXT("UbaCluster"), HordeCluster, GEngineIni);
+	GConfig->GetString(TEXT("Horde"), TEXT("UbaLocalHost"), HordeHost, GEngineIni);
+	GConfig->GetInt(TEXT("Horde"), TEXT("MaxCores"), HordeMaxCores, GEngineIni);
+	GConfig->GetBool(TEXT("Horde"), TEXT("UbaAllowWine"), bHordeAllowWine, GEngineIni);
 
 	FString UbaConnectionModeString;
 	if (GConfig->GetString(TEXT("Horde"), TEXT("UbaConnectionMode"), UbaConnectionModeString, GEngineIni))
 	{
-		LexFromString(OutConfig.HordeConnectionMode, *UbaConnectionModeString);
+		LexFromString(HordeConnectionMode, *UbaConnectionModeString);
 	}
 
 	FString UbaEncryptionString;
 	if (GConfig->GetString(TEXT("Horde"), TEXT("UbaEncryption"), UbaEncryptionString, GEngineIni))
 	{
-		LexFromString(OutConfig.HordeEncryption, *UbaEncryptionString);
+		LexFromString(HordeEncryption, *UbaEncryptionString);
 	}
 
-	GConfig->GetString(TEXT("Horde"), TEXT("UbaSentryUrl"), OutConfig.UbaSentryUrl, GEngineIni);
+	GConfig->GetString(TEXT("Horde"), TEXT("UbaSentryUrl"), UbaSentryUrl, GEngineIni);
 
 	[REDACTED_EMAIL]: enable this as soon as UBT and shader compiler can both enable Horde connection via the same configuration source
 	/*FString UbaEnabledString;
 	if (GConfig->GetString(TEXT("Horde"), TEXT("UbaEnabled"), UbaEnabledString, GEngineIni))
 	{
-		LexFromString(OutConfig.HordeEnabled, *UbaEnabledString);
+		LexFromString(HordeEnabled, *UbaEnabledString);
 	}*/
 }
 
@@ -173,30 +187,31 @@ static bool ParseCmdLineEnum(const TCHAR* CmdLine, const TCHAR* Parameter, T& Ou
 	return false;
 }
 
-static void ParseUbaHordeConfigFromCommandline(FUbaHordeConfig& OutConfig, const TCHAR* CmdLine)
+void FUbaHordeConfig::ParseConfigFromCommandline(const TCHAR* CmdLine)
 {
-	FParse::Value(CmdLine, TEXT("UBAHorde="), OutConfig.HordeServer);
-	FParse::Value(CmdLine, TEXT("UBAHordeToken="), OutConfig.HordeToken);
-	FParse::Value(CmdLine, TEXT("UBAHordePool="), OutConfig.OverrideHordePool);
-	ParseCmdLineBoolean(CmdLine, TEXT("UBAHordeAllowWine"), OutConfig.bHordeAllowWine);
-	FParse::Value(CmdLine, TEXT("UBAHordeMaxCores="), OutConfig.HordeMaxCores);
-	FParse::Value(CmdLine, TEXT("UBAHordeMaxLocalCores="), OutConfig.HordeMaxLocalCores);
-	FParse::Value(CmdLine, TEXT("UBAHordeHost="), OutConfig.HordeHost);
-	FParse::Value(CmdLine, TEXT("UBAHordeCluster="), OutConfig.HordeCluster);
-	ParseCmdLineEnum(CmdLine, TEXT("UBAHordeConnectionMode="), OutConfig.HordeConnectionMode);
-	ParseCmdLineEnum(CmdLine, TEXT("UBAHordeEncryption="), OutConfig.HordeEncryption);
-	FParse::Value(CmdLine, TEXT("UBASentryUrl="), OutConfig.UbaSentryUrl);
+	FParse::Value(CmdLine, TEXT("UBAHorde="), HordeServer);
+	FParse::Value(CmdLine, TEXT("UBAHordeToken="), HordeToken);
+	FParse::Value(CmdLine, TEXT("UBAHordePool="), OverrideHordePool);
+	ParseCmdLineBoolean(CmdLine, TEXT("UBAHordeAllowWine"), bHordeAllowWine);
+	FParse::Value(CmdLine, TEXT("UBAHordeMaxCores="), HordeMaxCores);
+	FParse::Value(CmdLine, TEXT("UBAHordeHost="), HordeHost);
+	FParse::Value(CmdLine, TEXT("UBAHordeCluster="), HordeCluster);
+	ParseCmdLineEnum(CmdLine, TEXT("UBAHordeConnectionMode="), HordeConnectionMode);
+	ParseCmdLineEnum(CmdLine, TEXT("UBAHordeEncryption="), HordeEncryption);
+	FParse::Value(CmdLine, TEXT("UBASentryUrl="), UbaSentryUrl);
 	if (FParse::Param(CmdLine, TEXT("UBADisableHorde")))
 	{
-		OutConfig.HordeEnabled = EUbaHordeEnabledMode::Disabled;
+		HordeEnabled = EUbaHordeEnabledMode::Disabled;
 	}
 	else if (FParse::Param(CmdLine, TEXT("UBAEnableHorde")))
 	{
-		OutConfig.HordeEnabled = EUbaHordeEnabledMode::Enabled;
+		HordeEnabled = EUbaHordeEnabledMode::Enabled;
 	}
+
+	FParse::Value(CmdLine, TEXT("MaxParallelActions="), MaxParallelActions);
 	if (FParse::Param(CmdLine, TEXT("ExclusiveRemoteShaderCompiling")))
 	{
-		OutConfig.HordeMaxLocalCores = 0; // -ExclusiveRemoteShaderCompiling overrides -UBAHordeMaxLocalCores
+		MaxParallelActions = 0; // -ExclusiveRemoteShaderCompiling overrides -MaxParallelActions
 	}
 }
 
@@ -255,30 +270,30 @@ static TArray<FString> GetBuildConfigFilePaths()
  *  2. From BaseEngine.ini configuration files under [Horde] section
  *  3. From command line argument (e.g. -UBAHorde=...)
  */
-static void ReadUbaHordeConfig(FUbaHordeConfig& OutConfig)
+void FUbaHordeConfig::Initialize()
 {
 	// Read configurations from BuildConfiguration.xml files
 	for (const FString& ConfigFilePath : GetBuildConfigFilePaths())
 	{
-		if (ParseUbaHordeConfigFromXmlFile(OutConfig, *ConfigFilePath))
+		if (ParseConfigFromXmlFile(*ConfigFilePath))
 		{
 			UE_LOG(LogUbaHorde, Verbose, TEXT("Read UBA/Horde settings from %s"), *ConfigFilePath);
 		}
 	}
 
 	// Override configuration entries by their configuration INI files
-	ParseUbaHordeConfigFromIniFile(OutConfig);
+	ParseConfigFromIniFile();
 
 	// Override configuration entries by their command line arguments
-	ParseUbaHordeConfigFromCommandline(OutConfig, FCommandLine::Get());
+	ParseConfigFromCommandline(FCommandLine::Get());
 
 	// If connection mode is "realy"-mode, encryption is implied and we use AES since UbaController does not support SSL/TLS yet
-	if (OutConfig.HordeConnectionMode == EUbaHordeConnectionMode::Relay && OutConfig.HordeEncryption == EUbaHordeEncryption::None)
+	if (HordeConnectionMode == EUbaHordeConnectionMode::Relay && HordeEncryption == EUbaHordeEncryption::None)
 	{
-		OutConfig.HordeEncryption = EUbaHordeEncryption::AES;
+		HordeEncryption = EUbaHordeEncryption::AES;
 	}
 
-	if (OutConfig.IsHordeEnabled())
+	if (IsHordeEnabled())
 	{
 		UE_LOG(
 			LogUbaHorde, Log,
@@ -292,15 +307,15 @@ static void ReadUbaHordeConfig(FUbaHordeConfig& OutConfig)
 			TEXT("\n  Allow Wine:        %s")
 			TEXT("\n  Mode:              %s")
 			TEXT("\n  Encryption:        %s")
-			, *OutConfig.HordeServer
-			, *OutConfig.GetHordePool()
-			, *OutConfig.HordeCondition
-			, *OutConfig.HordeCluster
-			, *OutConfig.HordeHost
-			, OutConfig.HordeMaxCores
-			, OutConfig.bHordeAllowWine ? TEXT("Yes") : TEXT("No")
-			, LexToString(OutConfig.HordeConnectionMode)
-			, LexToString(OutConfig.HordeEncryption)
+			, *HordeServer
+			, *GetHordePool()
+			, *HordeCondition
+			, *HordeCluster
+			, *HordeHost
+			, HordeMaxCores
+			, bHordeAllowWine ? TEXT("Yes") : TEXT("No")
+			, LexToString(HordeConnectionMode)
+			, LexToString(HordeEncryption)
 		);
 	}
 	else
@@ -332,7 +347,7 @@ UBACOORDINATORHORDE_API const FUbaHordeConfig& FUbaHordeConfig::Get()
 		FRWScopeLock LockGuard(InitializedLock, FRWScopeLockType::SLT_Write);
 		if (!bInitialized)
 		{
-			ReadUbaHordeConfig(Config);
+			Config.Initialize();
 			bInitialized = true;
 		}
 	}
diff --git a/Engine/Source/Developer/UbaCoordinatorHorde/Public/UbaHordeConfig.h b/Engine/Source/Developer/UbaCoordinatorHorde/Public/UbaHordeConfig.h
--- a/Engine/Source/Developer/UbaCoordinatorHorde/Public/UbaHordeConfig.h
+++ b/Engine/Source/Developer/UbaCoordinatorHorde/Public/UbaHordeConfig.h
@@ -29,8 +29,9 @@ enum class EUbaHordeEnabledMode
 };
 
 /** UBA/Horde configuration structure. This is meant to match the C# version UnrealBuildTool.UnrealBuildAcceleratorHordeConfig. */
-struct FUbaHordeConfig
+class FUbaHordeConfig
 {
+public:
 	static const TCHAR* ClusterDefault;
 	static const TCHAR* ClusterAuto;
 
@@ -71,7 +72,7 @@ struct FUbaHordeConfig
 	int32 HordeMaxCores = 576;
 
 	/** Maximum number of local CPU cores allowed to be used by build session. By default -1, which indicates to use as many CPU cores as the local machine can provide. */
-	int32 HordeMaxLocalCores = -1;
+	int32 MaxParallelActions = -1;
 
 	/** Allow use of POSIX/Wine. Only applicable to Horde agents running Linux. Can still be ignored if Wine executable is not set on agent. */
 	bool bHordeAllowWine = true;
@@ -93,6 +94,14 @@ struct FUbaHordeConfig
 
 	/** Returns the global UBA/Horde configuration and initializes it from a list of XML files on the first call. */
 	UBACOORDINATORHORDE_API static const FUbaHordeConfig& Get();
+
+private:
+	void Initialize();
+
+	bool ParseConfigFromXmlFile(const TCHAR* InBuildConfigFilename);
+	void ParseConfigFromIniFile();
+	void ParseConfigFromCommandline(const TCHAR* CmdLine);
+
 };
 
 UBACOORDINATORHORDE_API const TCHAR* LexToString(EUbaHordeConnectionMode InConnectionMode);
diff --git a/Engine/Source/Runtime/Engine/Private/ShaderCompiler/ShaderCompiler.cpp b/Engine/Source/Runtime/Engine/Private/ShaderCompiler/ShaderCompiler.cpp
--- a/Engine/Source/Runtime/Engine/Private/ShaderCompiler/ShaderCompiler.cpp
+++ b/Engine/Source/Runtime/Engine/Private/ShaderCompiler/ShaderCompiler.cpp
@@ -1014,9 +1014,13 @@ FShaderCompileThreadRunnableBase* FShaderCompilingManager::LaunchRemoteShaderCom
 
 	if (BuildDistributionController)
 	{
+		// Initialize distributed controller with worker limits if the controller also supports local workers
+		BuildDistributionController->SetMaxLocalWorkers(GetNumLocalWorkers());
+
 		// Allocate distributed shader compiling thread
 		UE_LOG(LogShaderCompilers, Display, TEXT("Using %s for shader compilation"), *BuildDistributionController->GetName());
 		TUniquePtr<FShaderCompileThreadRunnableBase> RemoteCompileThread = MakeUnique<FShaderCompileDistributedThreadRunnable_Interface>(this, *BuildDistributionController);
+
 		bExclusiveDistributedShaderCompiling = BuildDistributionController->SupportsLocalWorkers();
 		return LaunchShaderCompilingThread(MoveTemp(RemoteCompileThread), bDelayThreadExecution);
 	}
@@ -1206,6 +1210,12 @@ void FShaderCompilingManager::CalculateNumberOfCompilingThreads(int32 NumberOfCo
 void FShaderCompilingManager::OnMachineResourcesChanged(int32 NumberOfCores, int32 NumberOfCoresIncludingHyperthreads)
 {
 	CalculateNumberOfCompilingThreads(NumberOfCores, NumberOfCoresIncludingHyperthreads);
+
+	if (BuildDistributionController)
+	{
+		BuildDistributionController->SetMaxLocalWorkers(GetNumLocalWorkers());
+	}
+
 	for (const auto& Thread : Threads)
 	{
 		Thread->OnMachineResourcesChanged();
