[UBA] Enable Horde connection in UbaController through same [Horde]:UbaEnable INI config as for UBT.

UBA shader compilation has been slow in the last few days due to Horde helpers not being requested on build machines.
This also changes the log verbosity of "UBA/Horde Configuration" info from Log to Display.

#rnx

diff --git a/Engine/Source/Developer/UbaCoordinatorHorde/Private/UbaHordeConfig.cpp b/Engine/Source/Developer/UbaCoordinatorHorde/Private/UbaHordeConfig.cpp
--- a/Engine/Source/Developer/UbaCoordinatorHorde/Private/UbaHordeConfig.cpp
+++ b/Engine/Source/Developer/UbaCoordinatorHorde/Private/UbaHordeConfig.cpp
@@ -149,12 +149,11 @@ void FUbaHordeConfig::ParseConfigFromIniFile()
 
 	GConfig->GetString(TEXT("Horde"), TEXT("UbaSentryUrl"), UbaSentryUrl, GEngineIni);
 
-	[REDACTED_EMAIL]: enable this as soon as UBT and shader compiler can both enable Horde connection via the same configuration source
-	/*FString UbaEnabledString;
+	FString UbaEnabledString;
 	if (GConfig->GetString(TEXT("Horde"), TEXT("UbaEnabled"), UbaEnabledString, GEngineIni))
 	{
 		LexFromString(HordeEnabled, *UbaEnabledString);
-	}*/
+	}
 }
 
 static bool ParseCmdLineBoolean(const TCHAR* CmdLine, const TCHAR* Parameter, bool& bOutValue)
@@ -296,7 +295,7 @@ void FUbaHordeConfig::Initialize()
 	if (IsHordeEnabled())
 	{
 		UE_LOG(
-			LogUbaHorde, Log,
+			LogUbaHorde, Display,
 			TEXT("UBA/Horde Configuration:")
 			TEXT("\n  Server:            %s")
 			TEXT("\n  Pool:              %s")
@@ -320,7 +319,7 @@ void FUbaHordeConfig::Initialize()
 	}
 	else
 	{
-		UE_LOG(LogUbaHorde, Log, TEXT("UBA/Horde Configuration: Not Enabled"));
+		UE_LOG(LogUbaHorde, Display, TEXT("UBA/Horde Configuration: Not Enabled"));
 	}
 }
 
