[UBA] Use old [UbaController]:Enabled INI setting for compatibility reasons until UbaController can connect to Horde by default.

#rnx

diff --git a/Engine/Source/Developer/UbaCoordinatorHorde/Private/UbaHordeConfig.cpp b/Engine/Source/Developer/UbaCoordinatorHorde/Private/UbaHordeConfig.cpp
--- a/Engine/Source/Developer/UbaCoordinatorHorde/Private/UbaHordeConfig.cpp
+++ b/Engine/Source/Developer/UbaCoordinatorHorde/Private/UbaHordeConfig.cpp
@@ -149,8 +149,9 @@ void FUbaHordeConfig::ParseConfigFromIniFile()
 
 	GConfig->GetString(TEXT("Horde"), TEXT("UbaSentryUrl"), UbaSentryUrl, GEngineIni);
 
+	[REDACTED_EMAIL]: This is for compatibility reasons with the build farm
 	FString UbaEnabledString;
-	if (GConfig->GetString(TEXT("Horde"), TEXT("UbaEnabled"), UbaEnabledString, GEngineIni))
+	if (GConfig->GetString(TEXT("UbaController"), TEXT("Enabled"), UbaEnabledString, GEngineIni))
 	{
 		LexFromString(HordeEnabled, *UbaEnabledString);
 	}
