[UBA] Fix parsing of tri-state [UbaController]:Enabled.

#rnx

diff --git a/Engine/Source/Developer/UbaCoordinatorHorde/Private/UbaHordeConfig.cpp b/Engine/Source/Developer/UbaCoordinatorHorde/Private/UbaHordeConfig.cpp
--- a/Engine/Source/Developer/UbaCoordinatorHorde/Private/UbaHordeConfig.cpp
+++ b/Engine/Source/Developer/UbaCoordinatorHorde/Private/UbaHordeConfig.cpp
@@ -122,6 +122,22 @@ bool FUbaHordeConfig::ParseConfigFromXmlFile(const TCHAR* InBuildConfigFilename)
 	return true;
 }
 
+static bool GetConfigIdentifier(FConfigCacheIni& Config, const TCHAR* Section, const TCHAR* Key, FString& Value, const FString& Filename)
+{
+	if (Config.GetString(Section, Key, Value, Filename))
+	{
+		// Trim comments from string if this string was specified as 'Enabled=True ; False' instead of 'Enabled="True" ; False'
+		const int32 CommentPosition = Value.Find(TEXT(";"));
+		if (CommentPosition != INDEX_NONE)
+		{
+			Value.LeftInline(CommentPosition);
+			Value.TrimCharInline(TEXT(' '), nullptr);
+		}
+		return true;
+	}
+	return false;
+}
+
 void FUbaHordeConfig::ParseConfigFromIniFile()
 {
 	GConfig->GetString(TEXT("Horde"), TEXT("ServerUrl"), HordeServer, GEngineIni);
@@ -151,7 +167,7 @@ void FUbaHordeConfig::ParseConfigFromIniFile()
 
 	[REDACTED_EMAIL]: This is for compatibility reasons with the build farm
 	FString UbaEnabledString;
-	if (GConfig->GetString(TEXT("UbaController"), TEXT("Enabled"), UbaEnabledString, GEngineIni))
+	if (GetConfigIdentifier(*GConfig, TEXT("UbaController"), TEXT("Enabled"), UbaEnabledString, GEngineIni))
 	{
 		LexFromString(HordeEnabled, *UbaEnabledString);
 	}
