[UBA] Use "default" cluster in UBA/Horde config by default instead of "_auto".

This will not effect clients that configured the cluster explicitly to "_auto" in their BuildConfiguration.xml file.

#rnx

diff --git a/Engine/Source/Developer/UbaCoordinatorHorde/Public/UbaHordeConfig.h b/Engine/Source/Developer/UbaCoordinatorHorde/Public/UbaHordeConfig.h
--- a/Engine/Source/Developer/UbaCoordinatorHorde/Public/UbaHordeConfig.h
+++ b/Engine/Source/Developer/UbaCoordinatorHorde/Public/UbaHordeConfig.h
@@ -59,8 +59,8 @@ public:
 	/** Requirements for the Horde agent to assign. */
 	FString HordeCondition;
 
-	/** Compute cluster ID to use in Horde. Set to "_auto" to let Horde server resolve a suitable cluster. */
-	FString HordeCluster = FUbaHordeConfig::ClusterAuto;
+	/** Compute cluster ID to use in Horde. Set to "_auto" to let Horde server resolve a suitable cluster. By default "default". */
+	FString HordeCluster = FUbaHordeConfig::ClusterDefault;
 
 	/** Which IP address UBA server should assign to agents. */
 	FString HordeHost;
