[UBA] Also print HTTP response code alongside "invalid response body" error.

#rnx

diff --git a/Engine/Source/Developer/UbaCoordinatorHorde/Private/UbaHordeMetaClient.cpp b/Engine/Source/Developer/UbaCoordinatorHorde/Private/UbaHordeMetaClient.cpp
--- a/Engine/Source/Developer/UbaCoordinatorHorde/Private/UbaHordeMetaClient.cpp
+++ b/Engine/Source/Developer/UbaCoordinatorHorde/Private/UbaHordeMetaClient.cpp
@@ -112,7 +112,7 @@ TSharedPtr<FUbaHordeMetaClient::HordeClusterPromise, ESPMode::ThreadSafe> FUbaHo
 				ResponseCode == EHttpServerResponseCodes::TooManyRequests)
 			{
 				// Service Unavailable
-				UE_LOG(LogUbaHorde, Display, TEXT("Horde agent request returned with HTTP/%d: %s"), ResponseCode, *ResponseStr);
+				UE_LOG(LogUbaHorde, Display, TEXT("Horde agent request returned with HTTP/%d: %s"), (int32)ResponseCode, *ResponseStr);
 				return;
 			}
 
@@ -122,7 +122,7 @@ TSharedPtr<FUbaHordeMetaClient::HordeClusterPromise, ESPMode::ThreadSafe> FUbaHo
 			if (!FJsonSerializer::Deserialize(Reader, OutJson, FJsonSerializer::EFlags::None))
 			{
 				// Report invalid response body with Display verbosity only, since this should not fail a CIS job
-				UE_LOG(LogUbaHorde, Display, TEXT("Invalid response body: %s"), *ResponseStr);
+				UE_LOG(LogUbaHorde, Display, TEXT("Invalid response body (HTTP/%d): %s"), (int32)ResponseCode, *ResponseStr);
 				return;
 			}
 
@@ -130,7 +130,7 @@ TSharedPtr<FUbaHordeMetaClient::HordeClusterPromise, ESPMode::ThreadSafe> FUbaHo
 			if (!ClusterIdValue.Get())
 			{
 				// Report invalid response body with Display verbosity only, since this should not fail a CIS job
-				UE_LOG(LogUbaHorde, Display, TEXT("Invalid response body: %s"), *ResponseStr);
+				UE_LOG(LogUbaHorde, Display, TEXT("Invalid response body (HTTP/%d): %s"), (int32)ResponseCode, *ResponseStr);
 				return;
 			}
 
