[UBA] Improve error output on invalid response body from UBA/Horde connection.

#rnx

diff --git a/Engine/Source/Developer/UbaCoordinatorHorde/Private/UbaHordeMetaClient.cpp b/Engine/Source/Developer/UbaCoordinatorHorde/Private/UbaHordeMetaClient.cpp
--- a/Engine/Source/Developer/UbaCoordinatorHorde/Private/UbaHordeMetaClient.cpp
+++ b/Engine/Source/Developer/UbaCoordinatorHorde/Private/UbaHordeMetaClient.cpp
@@ -122,7 +122,7 @@ TSharedPtr<FUbaHordeMetaClient::HordeClusterPromise, ESPMode::ThreadSafe> FUbaHo
 			if (!FJsonSerializer::Deserialize(Reader, OutJson, FJsonSerializer::EFlags::None))
 			{
 				// Report invalid response body with Display verbosity only, since this should not fail a CIS job
-				UE_LOG(LogUbaHorde, Display, TEXT("Invalid response body (HTTP/%d): %s"), (int32)ResponseCode, *ResponseStr);
+				UE_LOG(LogUbaHorde, Display, TEXT("Invalid response body for cluster Id resolution (HTTP/%d): %s"), (int32)ResponseCode, *ResponseStr);
 				return;
 			}
 
@@ -130,7 +130,7 @@ TSharedPtr<FUbaHordeMetaClient::HordeClusterPromise, ESPMode::ThreadSafe> FUbaHo
 			if (!ClusterIdValue.Get())
 			{
 				// Report invalid response body with Display verbosity only, since this should not fail a CIS job
-				UE_LOG(LogUbaHorde, Display, TEXT("Invalid response body (HTTP/%d): %s"), (int32)ResponseCode, *ResponseStr);
+				UE_LOG(LogUbaHorde, Display, TEXT("Missing \"clusterId\" entry in response body (HTTP/%d): %s"), (int32)ResponseCode, *ResponseStr);
 				return;
 			}
 
@@ -192,7 +192,7 @@ TSharedPtr<FUbaHordeMetaClient::HordeMachinePromise, ESPMode::ThreadSafe> FUbaHo
 			if (!FJsonSerializer::Deserialize(Reader, OutJson, FJsonSerializer::EFlags::None))
 			{
 				// Report invalid response body with Display verbosity only, since this should not fail a CIS job
-				UE_LOG(LogUbaHorde, Display, TEXT("Invalid response body: %s"), *ResponseStr);
+				UE_LOG(LogUbaHorde, Display, TEXT("Invalid response body for remote helper request (HTTP/%d): %s"), (int32)ResponseCode, *ResponseStr);
 				return;
 			}
 
@@ -203,7 +203,7 @@ TSharedPtr<FUbaHordeMetaClient::HordeMachinePromise, ESPMode::ThreadSafe> FUbaHo
 			if (!NonceValue.Get() || !IpValue.Get() || !PortValue.Get())
 			{
 				// Report invalid response body with Display verbosity only, since this should not fail a CIS job
-				UE_LOG(LogUbaHorde, Display, TEXT("Invalid response body: %s"), *ResponseStr);
+				UE_LOG(LogUbaHorde, Display, TEXT("Missing \"nonce\", \"ip\", or \"port\" entry in response body (HTTP/%d): %s"), (int32)ResponseCode, *ResponseStr);
 				return;
 			}
 
