[UBA] Fix AES encrypted receiver function.

This fixes the receiver function of AES encrypted data for UbaController when queued data from a previous invocation is not large enough to satisfy the requested size.

#rnx

diff --git a/Engine/Source/Developer/UbaCoordinatorHorde/Private/UbaHordeComputeTransportAES.cpp b/Engine/Source/Developer/UbaCoordinatorHorde/Private/UbaHordeComputeTransportAES.cpp
--- a/Engine/Source/Developer/UbaCoordinatorHorde/Private/UbaHordeComputeTransportAES.cpp
+++ b/Engine/Source/Developer/UbaCoordinatorHorde/Private/UbaHordeComputeTransportAES.cpp
@@ -256,7 +256,15 @@ size_t FUbaHordeComputeTransportAES::Recv(void* Data, size_t Size)
 			RemainingData.Empty();
 		}
 
-		return (size_t)NumBytesFromRemainingData;
+		// If the remaining data satisfied the requested amount, exit early
+		if (NumBytesFromRemainingData == Size)
+		{
+			return (size_t)NumBytesFromRemainingData;
+		}
+
+		// ... Otherwise, adjust the output parameters and continue receiving more data
+		Size -= NumBytesFromRemainingData;
+		Data = reinterpret_cast<int8*>(Data) + NumBytesFromRemainingData;
 	}
 
 	// Receive message header
