[ubahorde] detect case of expired token (401 denied error) and refresh http client when encountered. fixes case of long running cooks failing to get helpers after the token expires.

diff --git a/Engine/Source/Developer/UbaCoordinatorHorde/Private/UbaHordeMetaClient.cpp b/Engine/Source/Developer/UbaCoordinatorHorde/Private/UbaHordeMetaClient.cpp
--- a/Engine/Source/Developer/UbaCoordinatorHorde/Private/UbaHordeMetaClient.cpp
+++ b/Engine/Source/Developer/UbaCoordinatorHorde/Private/UbaHordeMetaClient.cpp
@@ -34,6 +34,8 @@ bool FUbaHordeMetaClient::RefreshHttpClient()
 		return false;
 	}
 
+	bClientNeedsRefresh.store(false);
+
 	return true;
 }
 
@@ -83,6 +85,11 @@ TSharedPtr<FUbaHordeMetaClient::HordeClusterPromise, ESPMode::ThreadSafe> FUbaHo
 {
 	TSharedPtr<HordeClusterPromise> Promise = MakeShared<HordeClusterPromise, ESPMode::ThreadSafe>();
 
+	if (bClientNeedsRefresh)
+	{
+		RefreshHttpClient();
+	}
+
 	FHttpRequestRef Request = HttpClient->CreateRequest(TEXT("POST"), TEXT("api/v2/compute/_cluster"));
 
 	Request->SetDelegateThreadPolicy(EHttpRequestDelegateThreadPolicy::CompleteOnHttpThread);
@@ -116,6 +123,13 @@ TSharedPtr<FUbaHordeMetaClient::HordeClusterPromise, ESPMode::ThreadSafe> FUbaHo
 				return;
 			}
 
+			if (ResponseCode == EHttpServerResponseCodes::Denied)
+			{
+				UE_LOG(LogUbaHorde, Display, TEXT("Token expired, refreshing"));
+				bClientNeedsRefresh.store(true);
+				return;
+			}
+
 			TSharedPtr<FJsonValue> OutJson;
 			TSharedRef<TJsonReader<>> Reader = TJsonReaderFactory<>::Create(ResponseStr);
 
@@ -147,6 +161,11 @@ TSharedPtr<FUbaHordeMetaClient::HordeMachinePromise, ESPMode::ThreadSafe> FUbaHo
 {
 	TSharedPtr<HordeMachinePromise> Promise = MakeShared<HordeMachinePromise, ESPMode::ThreadSafe>();
 
+	if (bClientNeedsRefresh)
+	{
+		RefreshHttpClient();
+	}
+
 	const FString ResourcePath = FString::Format(TEXT("api/v2/compute/{0}"), { ClusterId != nullptr && *ClusterId != TEXT('\0') ? ClusterId : TEXT("default")});
 	FHttpRequestRef Request = HttpClient->CreateRequest(TEXT("POST"), *ResourcePath);
 
@@ -186,6 +205,13 @@ TSharedPtr<FUbaHordeMetaClient::HordeMachinePromise, ESPMode::ThreadSafe> FUbaHo
 				return;
 			}
 
+			if (ResponseCode == EHttpServerResponseCodes::Denied)
+			{
+				UE_LOG(LogUbaHorde, Display, TEXT("Token expired, refreshing"));
+				bClientNeedsRefresh.store(true);
+				return;
+			}
+
 			TSharedPtr<FJsonValue> OutJson;
 			TSharedRef<TJsonReader<>> Reader = TJsonReaderFactory<>::Create(ResponseStr);
 
diff --git a/Engine/Source/Developer/UbaCoordinatorHorde/Private/UbaHordeMetaClient.h b/Engine/Source/Developer/UbaCoordinatorHorde/Private/UbaHordeMetaClient.h
--- a/Engine/Source/Developer/UbaCoordinatorHorde/Private/UbaHordeMetaClient.h
+++ b/Engine/Source/Developer/UbaCoordinatorHorde/Private/UbaHordeMetaClient.h
@@ -98,4 +98,5 @@ public:
 private:
 	FString ServerUrl;
 	TUniquePtr<FHordeHttpClient> HttpClient;
+	std::atomic<bool> bClientNeedsRefresh{ true };
 };
