[ubahorde] add lease link to agent description when launching agents, can be useful to track down the cause of agent disconnects

diff --git a/Engine/Source/Developer/UbaCoordinatorHorde/Private/UbaHordeAgentManager.cpp b/Engine/Source/Developer/UbaCoordinatorHorde/Private/UbaHordeAgentManager.cpp
--- a/Engine/Source/Developer/UbaCoordinatorHorde/Private/UbaHordeAgentManager.cpp
+++ b/Engine/Source/Developer/UbaCoordinatorHorde/Private/UbaHordeAgentManager.cpp
@@ -433,6 +433,7 @@ void FUbaHordeAgentManager::ThreadAgent(FHordeAgentWrapper& Wrapper)
 
 		UbaAgentArgs.AddRaw("-Dir=%UE_HORDE_SHARED_DIR%\\Uba");
 		UbaAgentArgs.AddRaw("-Eventfile=%UE_HORDE_TERMINATION_SIGNAL_FILE%");
+		UbaAgentArgs.Add(FAnsiString::Printf("-Description=%s", TCHAR_TO_ANSI(*Agent->GetMachineInfo().LeaseLink)));
 
 		UbaAgentArgs.Finalize();
 
diff --git a/Engine/Source/Developer/UbaCoordinatorHorde/Private/UbaHordeMetaClient.cpp b/Engine/Source/Developer/UbaCoordinatorHorde/Private/UbaHordeMetaClient.cpp
--- a/Engine/Source/Developer/UbaCoordinatorHorde/Private/UbaHordeMetaClient.cpp
+++ b/Engine/Source/Developer/UbaCoordinatorHorde/Private/UbaHordeMetaClient.cpp
@@ -315,7 +315,8 @@ TSharedPtr<FUbaHordeMetaClient::HordeMachinePromise, ESPMode::ThreadSafe> FUbaHo
 			Info.Port = PortNumber;
 			FString::ToHexBlob(NonceString, Info.Nonce, HORDE_NONCE_SIZE);
 
-			FString HordeMachineSummary = FString::Printf(TEXT("UBA Horde machine assigned (%s) on '%s' [%s:%u]"), *OsFamily, *ClusterId, *Info.GetConnectionAddress(), (uint32)Info.GetConnectionPort().Port);
+			TStringBuilder<512> HordeMachineSummary;
+			HordeMachineSummary.Appendf(TEXT("UBA Horde machine assigned (%s) on '%s' [%s:%u]"), *OsFamily, *ClusterId, *Info.GetConnectionAddress(), (uint32)Info.GetConnectionPort().Port);
 
 #if WITH_SSL
 			if (TSharedPtr<FJsonValue> EncryptionValue = OutJson->AsObject()->TryGetField(TEXT("encryption")))
@@ -328,16 +329,16 @@ TSharedPtr<FUbaHordeMetaClient::HordeMachinePromise, ESPMode::ThreadSafe> FUbaHo
 						FString KeyString = KeyValue ? KeyValue->AsString() : FString();
 						if (KeyString.IsEmpty())
 						{
-							HordeMachineSummary += TEXT(" [AES key missing]");
+							HordeMachineSummary << TEXT(" [AES key missing]");
 						}
 						else if (KeyString.Len() != sizeof(Info.Key) * 2)
 						{
-							HordeMachineSummary += TEXT(" [AES key corrupted]");
+							HordeMachineSummary << TEXT(" [AES key corrupted]");
 						}
 						else
 						{
 							FString::ToHexBlob(KeyString, Info.Key, sizeof(Info.Key));
-							HordeMachineSummary += TEXT(" [AES key received]");
+							HordeMachineSummary << TEXT(" [AES key received]");
 						}
 					}
 				}
@@ -346,10 +347,11 @@ TSharedPtr<FUbaHordeMetaClient::HordeMachinePromise, ESPMode::ThreadSafe> FUbaHo
 
 			if (TSharedPtr<FJsonValue> LeaseIdValue = OutJson->AsObject()->TryGetField(TEXT("leaseId")))
 			{
-				HordeMachineSummary += FString::Format(TEXT(": {0}lease/{1}"), { this->ServerUrl, LeaseIdValue->AsString() });
+				Info.LeaseLink = FString::Format(TEXT("{0}lease/{1}"), { this->ServerUrl, LeaseIdValue->AsString() });
+				HordeMachineSummary << TEXT(": ") << Info.LeaseLink;
 			}
 
-			UE_LOG(LogUbaHorde, Display, TEXT("%s"), *HordeMachineSummary);
+			UE_LOG(LogUbaHorde, Display, TEXT("%s"), HordeMachineSummary.ToString());
 		});
 
 	Request->ProcessRequest();
diff --git a/Engine/Source/Developer/UbaCoordinatorHorde/Private/UbaHordeMetaClient.h b/Engine/Source/Developer/UbaCoordinatorHorde/Private/UbaHordeMetaClient.h
--- a/Engine/Source/Developer/UbaCoordinatorHorde/Private/UbaHordeMetaClient.h
+++ b/Engine/Source/Developer/UbaCoordinatorHorde/Private/UbaHordeMetaClient.h
@@ -47,6 +47,7 @@ struct FHordeRemoteMachineInfo
 	uint8 Key[HORDE_KEY_SIZE] = {};
 	TArray<uint8> CertificateData; // X.509 certificate for SSL connections in Pkcs12 format.
 	bool bRunsWindowOS = false;
+	FString LeaseLink;
 
 	static constexpr uint16 UbaPort = 7001;
 	static constexpr uint16 UbaProxyPort = 7002;
