[UbaController/UbaHorde]
* Added horde status in visualizer

diff --git a/Engine/Plugins/UbaController/Source/Private/UbaJobProcessor.cpp b/Engine/Plugins/UbaController/Source/Private/UbaJobProcessor.cpp
--- a/Engine/Plugins/UbaController/Source/Private/UbaJobProcessor.cpp
+++ b/Engine/Plugins/UbaController/Source/Private/UbaJobProcessor.cpp
@@ -423,13 +423,19 @@ void FUbaJobProcessor::StartUba()
 	// Only request Horde agents if Horde is enabled for UBA
 	if (FUbaHordeConfig::Get().IsHordeEnabled())
 	{
+		SessionServer_UpdateStatus(UbaSessionServer, 0, 1, TC("Horde"), uba::LogEntryType_Info, nullptr);
+		SessionServer_UpdateStatus(UbaSessionServer, 0, 6, TC("Starting"), uba::LogEntryType_Info, nullptr);
 		HordeAgentManager = MakeUnique<FUbaHordeAgentManager>(ControllerModule.GetWorkingDirectory(), GetUbaBinariesPath());
 
-		auto AddClientCallback = [](void* UserData, const uba::tchar* Ip, uint16 Port, const uba::tchar* Crypto16Characters)
+		HordeAgentManager->SetAddClientCallback([](void* UserData, const uba::tchar* Ip, uint16 Port, const uba::tchar* Crypto16Characters)
 			{
 				return NetworkServer_AddClient(reinterpret_cast<uba::NetworkServer*>(UserData), Ip, Port, Crypto16Characters);
-			};
-		HordeAgentManager->SetAddClientCallback(AddClientCallback, UbaServer);
+			}, UbaServer);
+
+		HordeAgentManager->SetUpdateStatusCallback([](void* UserData, const TCHAR* Status)
+			{
+				SessionServer_UpdateStatus(static_cast<uba::SessionServer*>(UserData), 0, 6, TCHAR_TO_UBASTRING(Status), uba::LogEntryType_Info, nullptr);
+			}, UbaSessionServer);
 	}
 
 	UE_LOG(LogUbaController, Display, TEXT("Created UBA storage server: RootDir=%s"), *RootDir);
diff --git a/Engine/Source/Developer/UbaCoordinatorHorde/Private/UbaHordeAgentManager.cpp b/Engine/Source/Developer/UbaCoordinatorHorde/Private/UbaHordeAgentManager.cpp
--- a/Engine/Source/Developer/UbaCoordinatorHorde/Private/UbaHordeAgentManager.cpp
+++ b/Engine/Source/Developer/UbaCoordinatorHorde/Private/UbaHordeAgentManager.cpp
@@ -28,6 +28,12 @@ namespace UbaCoordinatorHordeModule
 		TEXT("r.UbaHorde.ForwardAgentLogs"),
 		bHordeForwardAgentLogs,
 		TEXT("Enables or disables logging of stdout on agent side to show in controller log."));
+
+#if PLATFORM_MAC
+	constexpr uint32 EstimatedCoreCountPerInstance = 16;
+#else
+	constexpr uint32 EstimatedCoreCountPerInstance = 32;
+#endif
 }
 
 FUbaHordeAgentManager::FUbaHordeAgentManager(const FString& InWorkingDir, const FString& InBinariesPath)
@@ -37,6 +43,9 @@ FUbaHordeAgentManager::FUbaHordeAgentManager(const FString& InWorkingDir, const
 	,	TargetCoreCount(0)
 	,	EstimatedCoreCount(0)
 	,   ActiveCoreCount(0)
+	,	AgentsActive(0)
+	,	AgentsRequesting(0)
+	,	AgentsInProgress(0)
 	,	bAskForAgents(true)
 {
 }
@@ -86,6 +95,12 @@ void FUbaHordeAgentManager::SetAddClientCallback(FAddClientCallback* Callback, v
 	AddClientUserData = UserData;
 }
 
+void FUbaHordeAgentManager::SetUpdateStatusCallback(FUpdateStatusCallback* Callback, void* UserData)
+{
+	UpdateStatusCallback = Callback;
+	UpdateStatusUserData = UserData;
+}
+
 int32 FUbaHordeAgentManager::GetAgentCount() const
 {
 	FScopeLock AgentsScopeLock(&AgentsLock);
@@ -161,7 +176,7 @@ bool CreateHordeBundleFromFile(const std::filesystem::path& InputFilename, const
 
 void FUbaHordeAgentManager::RequestAgent()
 {
-	EstimatedCoreCount += 32; // We estimate a typical agent to have 32 cores
+	EstimatedCoreCount += UbaCoordinatorHordeModule::EstimatedCoreCountPerInstance;
 
 	const int32 NextAgentIndex = Agents.Num();
 
@@ -231,7 +246,7 @@ void FUbaHordeAgentManager::ThreadAgent(FHordeAgentWrapper& Wrapper)
 	FHordeRemoteMachineInfo::FPortInfo UbaPort, UbaProxyPort;
 
 	{
-		ON_SCOPE_EXIT{ EstimatedCoreCount -= 32; };
+		ON_SCOPE_EXIT{ EstimatedCoreCount -= UbaCoordinatorHordeModule::EstimatedCoreCountPerInstance; };
 
 		#if PLATFORM_WINDOWS
 		const char* AppName = "UbaAgent.exe";
@@ -269,6 +284,7 @@ void FUbaHordeAgentManager::ThreadAgent(FHordeAgentWrapper& Wrapper)
 				{
 					UE_LOG(LogUbaHorde, Error, TEXT("Failed to create Horde bundle for: %s"), *FilePath);
 					bAskForAgents = false;
+					UpdateStatus(TEXT("Failed to create bundle"));
 					return;
 				}
 				UE_LOG(LogUbaHorde, Display, TEXT("Created Horde bundle for: %s"), *FilePath);
@@ -278,14 +294,17 @@ void FUbaHordeAgentManager::ThreadAgent(FHordeAgentWrapper& Wrapper)
 
 		if (!HordeMetaClient)
 		{
+			UpdateStatus(TEXT("Connecting..."));
 			// Create Horde meta client right before we need it to make sure the CVar for the server URL has been read by now
 			HordeMetaClient = MakeUnique<FUbaHordeMetaClient>();
 			if (!HordeMetaClient->RefreshHttpClient())
 			{
 				UE_LOG(LogUbaHorde, Error, TEXT("Failed to create HttpClient for UbaAgent"));
 				bAskForAgents = false;
+				UpdateStatus(TEXT("Failed to connect"));
 				return;
 			}
+			UpdateStatus(TEXT("Running"));
 		}
 
 		if (!bAskForAgents)
@@ -293,6 +312,10 @@ void FUbaHordeAgentManager::ThreadAgent(FHordeAgentWrapper& Wrapper)
 			return;
 		}
 
+		++AgentsRequesting;
+		ON_SCOPE_EXIT{ --AgentsRequesting; UpdateStatus(nullptr); };
+		UpdateStatus(nullptr);
+
 		if (LastRequestFailTime == 0)
 		{
 			ScopeLock.Unlock();
@@ -311,6 +334,12 @@ void FUbaHordeAgentManager::ThreadAgent(FHordeAgentWrapper& Wrapper)
 			}
 		}
 
+		if (ActiveCoreCount >= TargetCoreCount)
+		{
+			UpdateStatus(nullptr);
+			return;
+		}
+
 		// Initialize horde request with connection mode preference and encryption.
 		// Pool ID is superseded by cluster ID, but it still has to be provided if default cluster and direct mode is used.
 		const FString HordeRequestJsonBody = HordeMetaClient->BuildHordeRequestJsonBody(Config.HordeConnectionMode == EUbaHordeConnectionMode::Direct ? Config.GetHordePool() : TEXT(""), Config.HordeConnectionMode, Config.HordeEncryption);
@@ -347,13 +376,21 @@ void FUbaHordeAgentManager::ThreadAgent(FHordeAgentWrapper& Wrapper)
 		{
 			UE_CLOG(!LastRequestFailTime, LogUbaHorde, Verbose, TEXT("No resources available in Horde. Will keep retrying until %u cores are used (Currently have %u)"), TargetCoreCount.Load(), ActiveCoreCount.Load());
 			LastRequestFailTime = FPlatformTime::Cycles64();
+			UpdateStatus(nullptr);
 			return;
 		}
 
 		LastRequestFailTime = 0;
 
+		++AgentsInProgress;
+		UpdateStatus(nullptr);
+
+		ON_SCOPE_EXIT{ --AgentsInProgress; UpdateStatus(nullptr); };
+
 		ScopeLock.Unlock();
 
+
+
 		if (ShouldExit.Wait(0))
 		{
 			return;
@@ -464,9 +501,12 @@ void FUbaHordeAgentManager::ThreadAgent(FHordeAgentWrapper& Wrapper)
 		MachineCoreCount = MachineInfo.LogicalCores;
 		EstimatedCoreCount += MachineCoreCount;
 		ActiveCoreCount += MachineCoreCount;
+		++AgentsActive;
+
+		UpdateStatus(nullptr);
 	}
 
-	ON_SCOPE_EXIT{ EstimatedCoreCount -= MachineCoreCount; ActiveCoreCount -= MachineCoreCount; };
+	ON_SCOPE_EXIT{ EstimatedCoreCount -= MachineCoreCount; ActiveCoreCount -= MachineCoreCount; --AgentsActive; UpdateStatus(nullptr); };
 
 	uint32 callCounter = 0; // TODO: This should react on the listen string instead of waiting for two text messages :)
 
@@ -510,3 +550,57 @@ void FUbaHordeAgentManager::ThreadAgent(FHordeAgentWrapper& Wrapper)
 	}
 }
 
+void FUbaHordeAgentManager::UpdateStatus(const TCHAR* Status)
+{
+	if (!UpdateStatusCallback)
+	{
+		return;
+	}
+
+	FScopeLock Lock(&UpdateStatusLock);
+	bool HasFailTime = LastRequestFailTime != 0;
+
+	if (HasFailTime == UpdateHadFailTime && AgentsRequesting == UpdateStatusAgentsRequesting &&
+		AgentsInProgress == UpdateStatusAgentsInProgress && AgentsActive == UpdateStatusAgentsActive && (!Status || UpdateStatusText == Status))
+	{
+		return;
+	}
+
+	UpdateHadFailTime = HasFailTime;
+	UpdateStatusAgentsRequesting = AgentsRequesting;
+	UpdateStatusAgentsInProgress = AgentsInProgress;
+	UpdateStatusAgentsActive = AgentsActive;
+
+	TStringBuilder<256> Temp;
+
+	if (Status)
+	{
+		UpdateStatusText = Status;
+	}
+	else
+	{
+		Temp.Appendf(TEXT("%s"), *UpdateStatusText);
+		if (UpdateStatusAgentsActive)
+		{
+			uint32 Acc = ActiveCoreCount;
+			Temp.Appendf(TEXT(". %u agents (%u cores)"), UpdateStatusAgentsActive, Acc);
+		}
+
+		if (UpdateHadFailTime)
+		{
+			Temp.Appendf(TEXT(" - No agents available."));
+		}
+		else if (UpdateStatusAgentsRequesting)
+		{
+			Temp.Appendf(TEXT(" - Requesting %u agents..."), UpdateStatusAgentsRequesting);
+		}
+
+		if (AgentsInProgress)
+		{
+			Temp.Appendf(TEXT(" (Preparing %u agents)"), (uint32)AgentsInProgress);
+		}
+		Status = *Temp;
+	}
+
+	UpdateStatusCallback(UpdateStatusUserData, Status);
+}
diff --git a/Engine/Source/Developer/UbaCoordinatorHorde/Public/UbaHordeAgentManager.h b/Engine/Source/Developer/UbaCoordinatorHorde/Public/UbaHordeAgentManager.h
--- a/Engine/Source/Developer/UbaCoordinatorHorde/Public/UbaHordeAgentManager.h
+++ b/Engine/Source/Developer/UbaCoordinatorHorde/Public/UbaHordeAgentManager.h
@@ -22,6 +22,9 @@ public:
 	using FAddClientCallback = bool(void* UserData, const uba::tchar* Ip, uint16 Port, const uba::tchar* Crypto16Characters);
 	UBACOORDINATORHORDE_API void SetAddClientCallback(FAddClientCallback* Callback, void* UserData);
 
+	using FUpdateStatusCallback = void(void* UserData, const TCHAR* Status);
+	UBACOORDINATORHORDE_API void SetUpdateStatusCallback(FUpdateStatusCallback* Callback, void* UserData);
+
 	// Returns the number of agents currently handled by this agent manager.
 	UBACOORDINATORHORDE_API int32 GetAgentCount() const;
 
@@ -37,6 +40,7 @@ private:
 
 	void RequestAgent();
 	void ThreadAgent(FHordeAgentWrapper& Wrapper);
+	void UpdateStatus(const TCHAR* Status);
 
 	FString WorkingDir;
 	FString BinariesPath;
@@ -53,11 +57,23 @@ private:
 	TAtomic<uint32> TargetCoreCount;
 	TAtomic<uint32> EstimatedCoreCount;
 	TAtomic<uint32> ActiveCoreCount;
+	TAtomic<uint32> AgentsActive;
+	TAtomic<uint32> AgentsRequesting;
+	TAtomic<uint32> AgentsInProgress;
 	TAtomic<bool> bAskForAgents;
 
 	FAddClientCallback* AddClientCallback = nullptr;
 	void* AddClientUserData = nullptr;
 
+	FCriticalSection UpdateStatusLock;
+	FString UpdateStatusText;
+	uint32 UpdateStatusAgentsInProgress = 0;
+	uint32 UpdateStatusAgentsActive = 0;
+	uint32 UpdateStatusAgentsRequesting = 0;
+	bool UpdateHadFailTime = false;
+	FUpdateStatusCallback* UpdateStatusCallback = nullptr;
+	void* UpdateStatusUserData = nullptr;
+
 	FUbaHordeAgentManager(const FUbaHordeAgentManager&) = delete;
 	void operator=(const FUbaHordeAgentManager&) = delete;
 };
