[UBA] Improve error reports of OpenSSL encryption/decryption in UBA/AES transport layer.

- Print warnings instead of verbose logging for failed OpenSSL calls.
- Include return code, error code, and error string representation in all reports of failed OpenSSL calls.
- Disable zero-padding, effectively using the default padding configuration for OpenSSL in UBA/AES transport layer.
- Flush and report all previoulsy queued OpenSSL errors in case they were unhandled.

#rnx

diff --git a/Engine/Plugins/UbaController/Source/Public/UbaControllerModule.h b/Engine/Plugins/UbaController/Source/Public/UbaControllerModule.h
--- a/Engine/Plugins/UbaController/Source/Public/UbaControllerModule.h
+++ b/Engine/Plugins/UbaController/Source/Public/UbaControllerModule.h
@@ -7,7 +7,7 @@
 class FUbaJobProcessor;
 struct FTaskCommandData;
 
-DECLARE_LOG_CATEGORY_EXTERN(LogUbaController, Log, Log);
+DECLARE_LOG_CATEGORY_EXTERN(LogUbaController, Log, All);
 
 class FUbaControllerModule : public IDistributedBuildController , public TSharedFromThis<FUbaControllerModule>
 {
diff --git a/Engine/Source/Developer/UbaCoordinatorHorde/Private/UbaHordeAgent.h b/Engine/Source/Developer/UbaCoordinatorHorde/Private/UbaHordeAgent.h
--- a/Engine/Source/Developer/UbaCoordinatorHorde/Private/UbaHordeAgent.h
+++ b/Engine/Source/Developer/UbaCoordinatorHorde/Private/UbaHordeAgent.h
@@ -8,7 +8,7 @@
 class FComputeSocket;
 class FAgentMessageChannel;
 
-DECLARE_LOG_CATEGORY_EXTERN(LogUbaHordeAgent, Log, Log);
+DECLARE_LOG_CATEGORY_EXTERN(LogUbaHordeAgent, Log, All);
 
 class FUbaHordeAgent
 {
diff --git a/Engine/Source/Developer/UbaCoordinatorHorde/Private/UbaHordeComputeTransportAES.cpp b/Engine/Source/Developer/UbaCoordinatorHorde/Private/UbaHordeComputeTransportAES.cpp
--- a/Engine/Source/Developer/UbaCoordinatorHorde/Private/UbaHordeComputeTransportAES.cpp
+++ b/Engine/Source/Developer/UbaCoordinatorHorde/Private/UbaHordeComputeTransportAES.cpp
@@ -6,6 +6,7 @@
 
 #include <openssl/evp.h>
 #include <openssl/aes.h>
+#include <openssl/err.h>
 
 
 // Message format: [length (4 bytes)][nonce (12 bytes)][encrypted data][tag (16 bytes)]
@@ -26,6 +27,12 @@ namespace CryptoFormatAES
 	static constexpr int32 HeaderAndTagSize = sizeof(FPacketHeader) + TagSize;
 } // AESFormat
 
+#define VERIFY_OPENSSL_CALL(FUNC, ...) \
+	VerifyOpenSSLStatus(FUNC(__VA_ARGS__), TEXT(#FUNC))
+
+#define VERIFY_OPENSSL_CONTEXT_CALL(CONTEXT, FUNC, ...) \
+	(CONTEXT)->VerifyOpenSSLStatus(FUNC(__VA_ARGS__), TEXT(#FUNC))
+
 struct FUbaHordeComputeTransportAES::FOpenSSLContext
 {
 	EVP_CIPHER_CTX* EncryptionContext = nullptr;
@@ -35,6 +42,8 @@ struct FUbaHordeComputeTransportAES::FOpenSSLContext
 	uint8 EncryptNonce[CryptoFormatAES::NonceSize] = {};
 	uint8 DecryptNonce[CryptoFormatAES::NonceSize] = {};
 
+	bool bHasErrors = false;
+
 	FOpenSSLContext(const uint8 (&InKey)[HORDE_KEY_SIZE])
 	{
 		FMemory::Memcpy(Key, InKey, sizeof(InKey));
@@ -44,6 +53,12 @@ struct FUbaHordeComputeTransportAES::FOpenSSLContext
 		{
 			NonceWord = (uint8)(FMath::Rand() % 256);
 		}
+
+		// Flush any potential OpenSSL error codes that were left in the queue before this context was established
+		while (unsigned long ErrorCode = ERR_get_error())
+		{
+			UE_LOG(LogUbaHorde, Warning, TEXT("Unhandled OpenSSL error (%lu) left in queue: %s"), ErrorCode, ANSI_TO_TCHAR(ERR_error_string(ErrorCode, nullptr)));
+		}
 	}
 
 	void UpdateNonce()
@@ -55,16 +70,25 @@ struct FUbaHordeComputeTransportAES::FOpenSSLContext
 		Nonce32Bits[2] = Nonce32Bits[0] ^ Nonce32Bits[1];
 	}
 
-	void EncryptBegin()
+	bool VerifyOpenSSLStatus(int ReturnCode, const TCHAR* CommandName)
+	{
+		if (ReturnCode != 1)
+		{
+			const unsigned long ErrorCode = ERR_get_error();
+			UE_LOG(LogUbaHorde, Warning, TEXT("%s() failed with return code (%d) and error code (%lu): %s"), CommandName, ReturnCode, ErrorCode, ANSI_TO_TCHAR(ERR_error_string(ErrorCode, nullptr)));
+			bHasErrors = true;
+			return false;
+		}
+		return true;
+	}
+
+	bool EncryptBegin()
 	{
 		// Update nonce to disguise equal messages with non-determinism
 		UpdateNonce();
 
 		// Use nonce as initialization vector
-		if (EVP_EncryptInit_ex(EncryptionContext, nullptr, nullptr, Key, EncryptNonce) == 0)
-		{
-			UE_LOG(LogUbaHorde, Verbose, TEXT("EVP_EncryptInit_ex() failed"));
-		}
+		return VERIFY_OPENSSL_CALL(EVP_EncryptInit_ex, EncryptionContext, nullptr, nullptr, Key, EncryptNonce);
 	}
 
 	int32 Copy(void* OutData, const void* InData, int32 InDataLength)
@@ -76,9 +100,9 @@ struct FUbaHordeComputeTransportAES::FOpenSSLContext
 	int32 Encrypt(void* OutDataCipher, const void* InData, int32 InDataLength)
 	{
 		int OutDataLength = 0;
-		if (EVP_EncryptUpdate(EncryptionContext, (unsigned char*)OutDataCipher, &OutDataLength, (const unsigned char*)InData, InDataLength) == 0)
+		if (!VERIFY_OPENSSL_CALL(EVP_EncryptUpdate, EncryptionContext, (unsigned char*)OutDataCipher, &OutDataLength, (const unsigned char*)InData, InDataLength))
 		{
-			UE_LOG(LogUbaHorde, Verbose, TEXT("EVP_EncryptUpdate(Size = %d) failed"), InDataLength);
+			return 0;
 		}
 		return (int32)OutDataLength;
 	}
@@ -88,15 +112,15 @@ struct FUbaHordeComputeTransportAES::FOpenSSLContext
 		int OutDataLength = 0;
 
 		int FooterLength = 0;
-		if (EVP_EncryptFinal_ex(EncryptionContext, (unsigned char*)OutData, &FooterLength) == 0)
+		if (!VERIFY_OPENSSL_CALL(EVP_EncryptFinal_ex, EncryptionContext, (unsigned char*)OutData, &FooterLength))
 		{
-			UE_LOG(LogUbaHorde, Verbose, TEXT("EVP_EncryptFinal_ex() failed"));
+			return 0;
 		}
 		OutDataLength += FooterLength;
 
-		if (EVP_CIPHER_CTX_ctrl(EncryptionContext, EVP_CTRL_GCM_GET_TAG, CryptoFormatAES::TagSize, ((char*)OutData) + OutDataLength) == 0)
+		if (!VERIFY_OPENSSL_CALL(EVP_CIPHER_CTX_ctrl, EncryptionContext, EVP_CTRL_GCM_GET_TAG, CryptoFormatAES::TagSize, ((char*)OutData) + OutDataLength))
 		{
-			UE_LOG(LogUbaHorde, Verbose, TEXT("EVP_CIPHER_CTX_ctrl(EVP_CTRL_GCM_GET_TAG) failed"));
+			return 0;
 		}
 		OutDataLength += CryptoFormatAES::TagSize;
 
@@ -107,18 +131,19 @@ struct FUbaHordeComputeTransportAES::FOpenSSLContext
 	{
 		int32 EncryptedDataLength = 0;
 
-		EncryptBegin();
-
-		uint8* EncryptedBytes = (uint8*)OutEncryptedData;
-		EncryptedDataLength += Copy(EncryptedBytes, &InDataLength, sizeof(InDataLength));
-		EncryptedDataLength += Copy(EncryptedBytes + EncryptedDataLength, EncryptNonce, sizeof(EncryptNonce));
-		EncryptedDataLength += Encrypt(EncryptedBytes + EncryptedDataLength, InData, InDataLength);
-		EncryptedDataLength += EncryptEnd(EncryptedBytes + EncryptedDataLength);
+		if (EncryptBegin())
+		{
+			uint8* EncryptedBytes = (uint8*)OutEncryptedData;
+			EncryptedDataLength += Copy(EncryptedBytes, &InDataLength, sizeof(InDataLength));
+			EncryptedDataLength += Copy(EncryptedBytes + EncryptedDataLength, EncryptNonce, sizeof(EncryptNonce));
+			EncryptedDataLength += Encrypt(EncryptedBytes + EncryptedDataLength, InData, InDataLength);
+			EncryptedDataLength += EncryptEnd(EncryptedBytes + EncryptedDataLength);
+		}
 
 		return EncryptedDataLength;
 	}
 
-	void DecryptBegin(const void* InNonce, int32 InNonceSize)
+	bool DecryptBegin(const void* InNonce, int32 InNonceSize)
 	{
 		// Use nonce that was received with the encrypted message
 		check(InNonce != nullptr);
@@ -126,18 +151,15 @@ struct FUbaHordeComputeTransportAES::FOpenSSLContext
 		FMemory::Memcpy(DecryptNonce, InNonce, (SIZE_T)InNonceSize);
 
 		// Use nonce as initialization vector
-		if (EVP_DecryptInit_ex(DecryptionContext, nullptr, nullptr, Key, DecryptNonce) == 0)
-		{
-			UE_LOG(LogUbaHorde, Verbose, TEXT("EVP_DecryptInit_ex() failed"));
-		}
+		return VERIFY_OPENSSL_CALL(EVP_DecryptInit_ex, DecryptionContext, nullptr, nullptr, Key, DecryptNonce);
 	}
 
 	int32 Decrypt(void* OutDataPlaintext, const void* InData, int32 InDataLength)
 	{
 		int OutDataLength = 0;
-		if (EVP_DecryptUpdate(DecryptionContext, (unsigned char*)OutDataPlaintext, &OutDataLength, (const unsigned char*)InData, InDataLength) == 0)
+		if (!VERIFY_OPENSSL_CALL(EVP_DecryptUpdate, DecryptionContext, (unsigned char*)OutDataPlaintext, &OutDataLength, (const unsigned char*)InData, InDataLength))
 		{
-			UE_LOG(LogUbaHorde, Verbose, TEXT("EVP_DecryptUpdate(Size = %d) failed"), InDataLength);
+			return 0;
 		}
 		return (int32)OutDataLength;
 	}
@@ -149,16 +171,16 @@ struct FUbaHordeComputeTransportAES::FOpenSSLContext
 
 		unsigned char Tag[CryptoFormatAES::TagSize];
 		FMemory::Memcpy(Tag, OutBytes, sizeof(Tag));
-		if (EVP_CIPHER_CTX_ctrl(DecryptionContext, EVP_CTRL_GCM_SET_TAG, CryptoFormatAES::TagSize, Tag) == 0)
+		if (!VERIFY_OPENSSL_CALL(EVP_CIPHER_CTX_ctrl, DecryptionContext, EVP_CTRL_GCM_SET_TAG, CryptoFormatAES::TagSize, Tag))
 		{
-			UE_LOG(LogUbaHorde, Verbose, TEXT("EVP_CIPHER_CTX_ctrl(EVP_CTRL_GCM_SET_TAG) failed"));
+			return 0;
 		}
 		OutDataLength += CryptoFormatAES::TagSize;
 
 		int FooterLength = 0;
-		if (EVP_DecryptFinal_ex(DecryptionContext, OutBytes + OutDataLength, &FooterLength) == 0)
+		if (!VERIFY_OPENSSL_CALL(EVP_DecryptFinal_ex, DecryptionContext, OutBytes + OutDataLength, &FooterLength))
 		{
-			UE_LOG(LogUbaHorde, Verbose, TEXT("EVP_DecryptFinal_ex() failed"));
+			return 0;
 		}
 		OutDataLength += FooterLength;
 
@@ -170,7 +192,6 @@ static EVP_CIPHER_CTX* NewAESCryptoContext()
 {
 	EVP_CIPHER_CTX* CryptoContext = EVP_CIPHER_CTX_new();
 	EVP_CIPHER_CTX_init(CryptoContext);
-	EVP_CIPHER_CTX_set_padding(CryptoContext, 0);
 	return CryptoContext;
 }
 
@@ -181,11 +202,12 @@ FUbaHordeComputeTransportAES::FUbaHordeComputeTransportAES(const FHordeRemoteMac
 	, bIsClosed(false)
 {
 	// Initialize AES encryption and decrption contexts
+	int Result = 0;
 	OpenSSLContext->EncryptionContext = NewAESCryptoContext();
-	EVP_EncryptInit_ex(OpenSSLContext->EncryptionContext, EVP_aes_256_gcm(), nullptr, nullptr, nullptr);
+	VERIFY_OPENSSL_CONTEXT_CALL(OpenSSLContext, EVP_EncryptInit_ex, OpenSSLContext->EncryptionContext, EVP_aes_256_gcm(), nullptr, nullptr, nullptr);
 
 	OpenSSLContext->DecryptionContext = NewAESCryptoContext();
-	EVP_DecryptInit_ex(OpenSSLContext->DecryptionContext, EVP_aes_256_gcm(), nullptr, nullptr, nullptr);
+	VERIFY_OPENSSL_CONTEXT_CALL(OpenSSLContext, EVP_DecryptInit_ex, OpenSSLContext->DecryptionContext, EVP_aes_256_gcm(), nullptr, nullptr, nullptr);
 }
 
 FUbaHordeComputeTransportAES::~FUbaHordeComputeTransportAES()
@@ -220,6 +242,10 @@ size_t FUbaHordeComputeTransportAES::Send(const void* Data, size_t Size)
 	uint8* EncryptedData = GetAndResizeEncryptedBuffer(MessageLength);
 
 	int32 EncryptedDataLength = OpenSSLContext->EncryptMessage(EncryptedData, Data, DataLength);
+	if (EncryptedDataLength == 0)
+	{
+		return 0;
+	}
 
 	// Send encrypted data over common transport layer
 	if (!InnerTransport->SendMessage(EncryptedData, (size_t)EncryptedDataLength))
@@ -280,7 +306,10 @@ size_t FUbaHordeComputeTransportAES::Recv(void* Data, size_t Size)
 		return 0;
 	}
 
-	OpenSSLContext->DecryptBegin(Header.Nonce, sizeof(Header.Nonce));
+	if (!OpenSSLContext->DecryptBegin(Header.Nonce, sizeof(Header.Nonce)))
+	{
+		return 0;
+	}
 
 	// Decrypt data and store remaining encrypted data for next Recv() call
 	const int32 NumBytesProcessed = FMath::Min(Header.DataLength, (int32)Size);
@@ -323,7 +352,7 @@ void FUbaHordeComputeTransportAES::Close()
 
 bool FUbaHordeComputeTransportAES::IsValid() const
 {
-	return InnerTransport.IsValid() && InnerTransport->IsValid();
+	return InnerTransport.IsValid() && InnerTransport->IsValid() && OpenSSLContext != nullptr && !OpenSSLContext->bHasErrors;
 }
 
 uint8* FUbaHordeComputeTransportAES::GetAndResizeEncryptedBuffer(int32 Size)
@@ -335,5 +364,8 @@ uint8* FUbaHordeComputeTransportAES::GetAndResizeEncryptedBuffer(int32 Size)
 	return EncryptedBuffer.GetData();
 }
 
+#undef VERIFY_OPENSSL_CALL
+#undef VERIFY_OPENSSL_CONTEXT_CALL
+
 #endif // WITH_SSL
 
diff --git a/Engine/Source/Developer/UbaCoordinatorHorde/Private/UbaHordeMetaClient.h b/Engine/Source/Developer/UbaCoordinatorHorde/Private/UbaHordeMetaClient.h
--- a/Engine/Source/Developer/UbaCoordinatorHorde/Private/UbaHordeMetaClient.h
+++ b/Engine/Source/Developer/UbaCoordinatorHorde/Private/UbaHordeMetaClient.h
@@ -10,7 +10,7 @@
 #include "UbaHordeAgentManager.h"
 #include "UbaHordeConfig.h"
 
-DECLARE_LOG_CATEGORY_EXTERN(LogUbaHorde, Log, Log);
+DECLARE_LOG_CATEGORY_EXTERN(LogUbaHorde, Log, All);
 
 #define HORDE_NONCE_SIZE 64
 #define HORDE_KEY_SIZE 32
