[UBA] Adopt new UBA/Horde configuration parsing in UbaController from UBT.

- Move parsing of UBA/Horde configurations from multiple sources into FUbaHordeConfigParser; This cleans up duplicated code that converts strings to their respective enums (EUbaHordeConnectionMode, EUbaHordeEncryption).
- Accept parameters withou the "Uba" prefix but keep the old syntax for compatibility, e.g. "UbaPool" in INI file can also be "Pool" now; This aligns with the new UBA/UBT syntax.
- Deprecate platform specific pool identifiers. These are now specified by their respective platform BaseEngine.ini override.
- Log information from what source parameters were parsed, e.g. "Cluster: _auto" can now be logged as "Cluster: _auto (INI)" to indicate it was accepted from the INI configuration; This helps identify misconfigurations as these parameters can come from four different sources (Default, XML, INI, command-line).

diff --git a/Engine/Plugins/UbaController/Source/Private/UbaControllerModule.cpp b/Engine/Plugins/UbaController/Source/Private/UbaControllerModule.cpp
--- a/Engine/Plugins/UbaController/Source/Private/UbaControllerModule.cpp
+++ b/Engine/Plugins/UbaController/Source/Private/UbaControllerModule.cpp
@@ -102,24 +102,7 @@ static bool IsUbaControllerEnabled()
 	}
 
 	// Check if UbaController is enabled via INI configuration in [UbaController] section.
-	FString EnabledState;
-	GConfig->GetString(TEXT("UbaController"), TEXT("Enabled"), EnabledState, GEngineIni);
-
-	// This "Enabled" parameter is a tri-state so we have to parse it as a string.
-	// Those strings can include the comments (starting with ';') from the INI file, so we have to trim that section from the string.
-	const int32 EnabledStateCommentPosition = EnabledState.Find(TEXT(";"));
-	if (EnabledStateCommentPosition != INDEX_NONE)
-	{
-		EnabledState.RemoveAt(EnabledStateCommentPosition, EnabledState.Len() - EnabledStateCommentPosition);
-		EnabledState.RemoveSpacesInline();
-	}
-
-	if (EnabledState.Equals(TEXT("True"), ESearchCase::IgnoreCase) || (EnabledState.Equals(TEXT("BuildMachineOnly"), ESearchCase::IgnoreCase) && GIsBuildMachine))
-	{
-		return true;
-	}
-
-	return false;
+	return FUbaHordeConfig::Get().bIsProviderEnabled;
 }
 
 bool FUbaControllerModule::IsSupported()
diff --git a/Engine/Plugins/UbaController/Source/Private/UbaJobProcessor.cpp b/Engine/Plugins/UbaController/Source/Private/UbaJobProcessor.cpp
--- a/Engine/Plugins/UbaController/Source/Private/UbaJobProcessor.cpp
+++ b/Engine/Plugins/UbaController/Source/Private/UbaJobProcessor.cpp
@@ -418,7 +418,7 @@ void FUbaJobProcessor::StartUba()
 	}
 
 	// Only request Horde agents if Horde is enabled for UBA
-	if (FUbaHordeConfig::Get().IsHordeEnabled())
+	if (FUbaHordeConfig::Get().bIsProviderEnabled)
 	{
 		SessionServer_UpdateStatus(UbaSessionServer, 0, 1, TC("Horde"), uba::LogEntryType_Info, nullptr);
 		SessionServer_UpdateStatus(UbaSessionServer, 0, 6, TC("Starting"), uba::LogEntryType_Info, nullptr);
diff --git a/Engine/Source/Developer/UbaCoordinatorHorde/Private/UbaHordeConfig.cpp b/Engine/Source/Developer/UbaCoordinatorHorde/Private/UbaHordeConfig.cpp
--- a/Engine/Source/Developer/UbaCoordinatorHorde/Private/UbaHordeConfig.cpp
+++ b/Engine/Source/Developer/UbaCoordinatorHorde/Private/UbaHordeConfig.cpp
@@ -15,33 +15,129 @@ const TCHAR* FUbaHordeConfig::ClusterAuto = TEXT("_auto");
 
 UBACOORDINATORHORDE_API const FString& FUbaHordeConfig::GetHordePool() const
 {
-	if (!OverrideHordePool.IsEmpty())
+	return HordePool;
+}
+
+struct FUbaHordeConfigParser
+{
+	template <typename T>
+	struct TParameter
+	{
+		T Value{};
+		const TCHAR* Source{};
+	};
+
+	const TCHAR* CurrentSource{};
+
+	TParameter<FString> Provider;
+	TParameter<FString> ProviderEnabled;
+	TParameter<FString> Pool;
+	TParameter<FString> Server;
+	TParameter<FString> Token;
+	TParameter<FString> Condition;
+	TParameter<FString> Cluster;
+	TParameter<FString> Host;
+	TParameter<FString> SentryUrl;
+	TParameter<int32> MaxCores;
+	TParameter<int32> MaxParallelActions;
+	TParameter<bool> bAllowWine;
+	TParameter<FString> ConnectionMode;
+	TParameter<FString> Encryption;
+
+	template <typename T>
+	void SetConfig(TParameter<T>& OutParameter, T InValue)
 	{
-		return OverrideHordePool;
+		OutParameter.Value = MoveTemp(InValue);
+		OutParameter.Source = CurrentSource;
 	}
-#if PLATFORM_WINDOWS
-	if (!WindowsHordePool.IsEmpty())
+
+	template <typename T>
+	void GetConfig(TParameter<T>& Parameter, T& OutValue)
 	{
-		return WindowsHordePool;
+		if (Parameter.Source)
+		{
+			OutValue = MoveTemp(Parameter.Value);
+		}
 	}
-#endif
-#if PLATFORM_MAC
-	if (!MacHordePool.IsEmpty())
+
+	bool GetUbaConfigString(const TCHAR* Section, const TCHAR* Key, TParameter<FString>& OutParameter)
 	{
-		return MacHordePool;
+		if (GConfig->GetString(Section, Key, OutParameter.Value, GEngineIni) || GConfig->GetString(Section, *FString::Printf(TEXT("Uba%s"), Key), OutParameter.Value, GEngineIni))
+		{
+			OutParameter.Source = CurrentSource;
+			return true;
+		}
+		return false;
 	}
-#endif
-#if PLATFORM_LINUX
-	if (!LinuxHordePool.IsEmpty())
+
+	bool GetUbaConfigInt(const TCHAR* Section, const TCHAR* Key, TParameter<int32>& OutParameter)
 	{
-		return LinuxHordePool;
+		if (GConfig->GetInt(Section, Key, OutParameter.Value, GEngineIni) || GConfig->GetInt(Section, *FString::Printf(TEXT("Uba%s"), Key), OutParameter.Value, GEngineIni))
+		{
+			OutParameter.Source = CurrentSource;
+			return true;
+		}
+		return false;
 	}
-#endif
-	return DefaultHordePool;
-}
 
-bool FUbaHordeConfig::ParseConfigFromXmlFile(const TCHAR* InBuildConfigFilename)
+	bool GetUbaConfigBool(const TCHAR* Section, const TCHAR* Key, TParameter<bool>& OutParameter)
+	{
+		if (GConfig->GetBool(Section, Key, OutParameter.Value, GEngineIni) || GConfig->GetBool(Section, *FString::Printf(TEXT("Uba%s"), Key), OutParameter.Value, GEngineIni))
+		{
+			OutParameter.Source = CurrentSource;
+			return true;
+		}
+		return false;
+	}
+
+	bool ParseCmdLineString(const TCHAR* CmdLine, const TCHAR* Parameter, TParameter<FString>& OutParameter)
+	{
+		if (FParse::Value(CmdLine, Parameter, OutParameter.Value))
+		{
+			OutParameter.Source = CurrentSource;
+			return true;
+		}
+		return false;
+	}
+
+	bool ParseCmdLineInt(const TCHAR* CmdLine, const TCHAR* Parameter, TParameter<int32>& OutParameter)
+	{
+		if (FParse::Value(CmdLine, Parameter, OutParameter.Value))
+		{
+			OutParameter.Source = CurrentSource;
+			return true;
+		}
+		return false;
+	}
+
+	bool ParseCmdLineBoolean(const TCHAR* CmdLine, const TCHAR* Parameter, TParameter<bool>& OutParameter)
+	{
+		if (FParse::Param(CmdLine, Parameter))
+		{
+			FString ArgumentValue;
+			if (FParse::Value(CmdLine, *FString::Printf(TEXT("%s="), Parameter), ArgumentValue))
+			{
+				OutParameter.Value = ArgumentValue != TEXT("false") && ArgumentValue != TEXT("0");
+			}
+			else
+			{
+				OutParameter.Value = true;
+			}
+			OutParameter.Source = CurrentSource;
+			return true;
+		}
+		return false;
+	}
+
+	bool ParseConfigFromXmlFile(const TCHAR* InBuildConfigFilename);
+	void ParseConfigFromIniFile();
+	void ParseConfigFromCommandline(const TCHAR* CmdLine);
+};
+
+bool FUbaHordeConfigParser::ParseConfigFromXmlFile(const TCHAR* InBuildConfigFilename)
 {
+	CurrentSource = TEXT("XML");
+
 	FXmlFile XmlFile;
 	if (!XmlFile.LoadFile(InBuildConfigFilename))
 	{
@@ -50,7 +146,6 @@ bool FUbaHordeConfig::ParseConfigFromXmlFile(const TCHAR* InBuildConfigFilename)
 
 	static const FString ConfigurationTag(TEXT("Configuration"));
 	static const FString BuildConfigurationTag(TEXT("BuildConfiguration"));
-	static const FString AllowUBAExecutorTag(TEXT("bAllowUBAExecutor"));
 	static const FString MaxParallelActionsTag(TEXT("MaxParallelActions"));
 	static const FString HordeTag(TEXT("Horde"));
 	static const FString UbaDisabledTag(TEXT("bDisableHorde"));
@@ -58,8 +153,7 @@ bool FUbaHordeConfig::ParseConfigFromXmlFile(const TCHAR* InBuildConfigFilename)
 	static const FString ConnectionModeTag(TEXT("ConnectionMode"));
 	static const FString EncryptionTag(TEXT("Encryption"));
 	static const FString ClusterTag(TEXT("Cluster"));
-	static const FString DefaultPoolTag(TEXT("Pool"));
-	static const FString MacPoolTag(TEXT("MacPool"));
+	static const FString PoolTag(TEXT("Pool"));
 	static const FString UBASentryUrlTag(TEXT("UBASentryUrl"));
 
 	const FXmlNode* ConfigurationNode = XmlFile.GetRootNode();
@@ -70,13 +164,9 @@ bool FUbaHordeConfig::ParseConfigFromXmlFile(const TCHAR* InBuildConfigFilename)
 
 	if (const FXmlNode* BuildConfigurationNode = ConfigurationNode->FindChildNode(BuildConfigurationTag))
 	{
-		/*if (const FXmlNode* AllowUBAExecutorNode = BuildConfigurationNode->FindChildNode(AllowUBAExecutorTag))
-		{
-			[REDACTED_EMAIL]: Once shader compiler and UBT are menat to use the same configuration source to enable UBA, make use of this XML tag
-		}*/
 		if (const FXmlNode* MaxParallelActionsNode = BuildConfigurationNode->FindChildNode(MaxParallelActionsTag))
 		{
-			MaxParallelActions = FCString::Atoi(*MaxParallelActionsNode->GetContent());
+			SetConfig(MaxParallelActions, FCString::Atoi(*MaxParallelActionsNode->GetContent()));
 		}
 	}
 
@@ -86,148 +176,98 @@ bool FUbaHordeConfig::ParseConfigFromXmlFile(const TCHAR* InBuildConfigFilename)
 		{
 			if (UbaDisabledNode->GetContent() == TEXT("true"))
 			{
-				HordeEnabled = EUbaHordeEnabledMode::Disabled;
+				SetConfig(ProviderEnabled, FString(TEXT("Disabled")));
 			}
 		}
-		if (const FXmlNode* DefaultPoolNode = HordeNode->FindChildNode(DefaultPoolTag))
+		if (const FXmlNode* PoolNode = HordeNode->FindChildNode(PoolTag))
 		{
-			DefaultHordePool = DefaultPoolNode->GetContent();
-		}
-		if (const FXmlNode* MacPoolNode = HordeNode->FindChildNode(MacPoolTag))
-		{
-			MacHordePool = MacPoolNode->GetContent();
+			SetConfig(Pool, PoolNode->GetContent());
 		}
 		if (const FXmlNode* ServerNode = HordeNode->FindChildNode(ServerTag))
 		{
-			HordeServer = ServerNode->GetContent();
+			SetConfig(Server, ServerNode->GetContent());
 		}
 		if (const FXmlNode* ConnectionModeNode = HordeNode->FindChildNode(ConnectionModeTag))
 		{
-			LexFromString(HordeConnectionMode, *ConnectionModeNode->GetContent());
+			SetConfig(ConnectionMode, ConnectionModeNode->GetContent());
 		}
 		if (const FXmlNode* EncryptionNode = HordeNode->FindChildNode(EncryptionTag))
 		{
-			LexFromString(HordeEncryption, *EncryptionNode->GetContent());
+			SetConfig(Encryption, EncryptionNode->GetContent());
 		}
 		if (const FXmlNode* ClusterNode = HordeNode->FindChildNode(ClusterTag))
 		{
-			HordeCluster = ClusterNode->GetContent();
+			SetConfig(Cluster, ClusterNode->GetContent());
 		}
 		if (const FXmlNode* UBASentryUrlNode = HordeNode->FindChildNode(UBASentryUrlTag))
 		{
-			UbaSentryUrl = UBASentryUrlNode->GetContent();
+			SetConfig(SentryUrl, UBASentryUrlNode->GetContent());
 		}
 	}
 
 	return true;
 }
 
-static bool GetConfigIdentifier(FConfigCacheIni& Config, const TCHAR* Section, const TCHAR* Key, FString& Value, const FString& Filename)
-{
-	if (Config.GetString(Section, Key, Value, Filename))
-	{
-		// Trim comments from string if this string was specified as 'Enabled=True ; False' instead of 'Enabled="True" ; False'
-		const int32 CommentPosition = Value.Find(TEXT(";"));
-		if (CommentPosition != INDEX_NONE)
-		{
-			Value.LeftInline(CommentPosition);
-			Value.TrimCharInline(TEXT(' '), nullptr);
-		}
-		return true;
-	}
-	return false;
-}
-
-void FUbaHordeConfig::ParseConfigFromIniFile()
+void FUbaHordeConfigParser::ParseConfigFromIniFile()
 {
-	GConfig->GetString(TEXT("Horde"), TEXT("ServerUrl"), HordeServer, GEngineIni);
-	GConfig->GetString(TEXT("Horde"), TEXT("UbaToken"), HordeToken, GEngineIni);
-	GConfig->GetString(TEXT("Horde"), TEXT("UbaPool"), DefaultHordePool, GEngineIni);
-	GConfig->GetString(TEXT("Horde"), TEXT("UbaLinuxPool"), LinuxHordePool, GEngineIni);
-	GConfig->GetString(TEXT("Horde"), TEXT("UbaMacPool"), MacHordePool, GEngineIni);
-	GConfig->GetString(TEXT("Horde"), TEXT("UbaWindowsPool"), WindowsHordePool, GEngineIni);
-	GConfig->GetString(TEXT("Horde"), TEXT("UbaCluster"), HordeCluster, GEngineIni);
-	GConfig->GetString(TEXT("Horde"), TEXT("UbaLocalHost"), HordeHost, GEngineIni);
-	GConfig->GetInt(TEXT("Horde"), TEXT("MaxCores"), HordeMaxCores, GEngineIni);
-	GConfig->GetBool(TEXT("Horde"), TEXT("UbaAllowWine"), bHordeAllowWine, GEngineIni);
+	CurrentSource = TEXT("INI");
 
-	FString UbaConnectionModeString;
-	if (GConfig->GetString(TEXT("Horde"), TEXT("UbaConnectionMode"), UbaConnectionModeString, GEngineIni))
-	{
-		LexFromString(HordeConnectionMode, *UbaConnectionModeString);
-	}
+	// Read array of Horde providers and use default is none is set. Only select the first available since UbaController does not support multiple providers.
+	TArray<FString> HordeProviders;
+	GConfig->GetArray(TEXT("UbaController"), GIsBuildMachine ? TEXT("BuildMachineProviders") : TEXT("Providers"), HordeProviders, GEngineIni);
+	SetConfig(Provider, HordeProviders.Num() > 0 ? HordeProviders[0] : TEXT("Uba.Provider.Horde"));
 
-	FString UbaEncryptionString;
-	if (GConfig->GetString(TEXT("Horde"), TEXT("UbaEncryption"), UbaEncryptionString, GEngineIni))
-	{
-		LexFromString(HordeEncryption, *UbaEncryptionString);
-	}
-
-	GConfig->GetString(TEXT("Horde"), TEXT("UbaSentryUrl"), UbaSentryUrl, GEngineIni);
-
-	[REDACTED_EMAIL]: This is for compatibility reasons with the build farm
-	FString UbaEnabledString;
-	if (GetConfigIdentifier(*GConfig, TEXT("UbaController"), TEXT("Enabled"), UbaEnabledString, GEngineIni))
-	{
-		LexFromString(HordeEnabled, *UbaEnabledString);
-	}
-}
+	GetUbaConfigString(*Provider.Value, TEXT("Enabled"), ProviderEnabled);
 
-static bool ParseCmdLineBoolean(const TCHAR* CmdLine, const TCHAR* Parameter, bool& bOutValue)
-{
-	if (FParse::Param(CmdLine, Parameter))
+	// Also check [UbaController]:Enabled for compatibility
+	if (GetUbaConfigString(TEXT("UbaController"), TEXT("Enabled"), ProviderEnabled))
 	{
-		FString ArgumentStringValue;
-		if (FParse::Value(CmdLine, *FString::Printf(TEXT("%s="), Parameter), ArgumentStringValue))
-		{
-			bOutValue = ArgumentStringValue != TEXT("false") && ArgumentStringValue != TEXT("0");
-		}
-		else
-		{
-			bOutValue = true;
-		}
-		return true;
+		UE_LOG(LogUbaHorde, Warning, TEXT("Configuration '[UbaController]:Enabled' is deprecated since UE 5.7; Use '[%s]:Enabled' instead"), *Provider.Value);
 	}
-	return false;
-}
 
-template <typename T>
-static bool ParseCmdLineEnum(const TCHAR* CmdLine, const TCHAR* Parameter, T& OutValue)
-{
-	FString ArgumentStringValue;
-	if (FParse::Value(CmdLine, Parameter, ArgumentStringValue))
-	{
-		LexFromString(OutValue, *ArgumentStringValue);
-		return true;
-	}
-	return false;
+	// Read all configurations from provider section
+	GetUbaConfigString(*Provider.Value, TEXT("ServerUrl"), Server);
+	GetUbaConfigString(*Provider.Value, TEXT("Token"), Token);
+	GetUbaConfigString(*Provider.Value, TEXT("Pool"), Pool);
+	GetUbaConfigString(*Provider.Value, TEXT("Cluster"), Cluster);
+	GetUbaConfigString(*Provider.Value, TEXT("LocalHost"), Host);
+	GetUbaConfigInt(*Provider.Value, TEXT("MaxCores"), MaxCores);
+	GetUbaConfigBool(*Provider.Value, TEXT("AllowWine"), bAllowWine);
+	GetUbaConfigString(*Provider.Value, TEXT("ConnectionMode"), ConnectionMode);
+	GetUbaConfigString(*Provider.Value, TEXT("Encryption"), Encryption);
+	GetUbaConfigString(*Provider.Value, TEXT("SentryUrl"), SentryUrl);
 }
 
-void FUbaHordeConfig::ParseConfigFromCommandline(const TCHAR* CmdLine)
+void FUbaHordeConfigParser::ParseConfigFromCommandline(const TCHAR* CmdLine)
 {
-	FParse::Value(CmdLine, TEXT("UBAHorde="), HordeServer);
-	FParse::Value(CmdLine, TEXT("UBAHordeToken="), HordeToken);
-	FParse::Value(CmdLine, TEXT("UBAHordePool="), OverrideHordePool);
-	ParseCmdLineBoolean(CmdLine, TEXT("UBAHordeAllowWine"), bHordeAllowWine);
-	FParse::Value(CmdLine, TEXT("UBAHordeMaxCores="), HordeMaxCores);
-	FParse::Value(CmdLine, TEXT("UBAHordeHost="), HordeHost);
-	FParse::Value(CmdLine, TEXT("UBAHordeCluster="), HordeCluster);
-	ParseCmdLineEnum(CmdLine, TEXT("UBAHordeConnectionMode="), HordeConnectionMode);
-	ParseCmdLineEnum(CmdLine, TEXT("UBAHordeEncryption="), HordeEncryption);
-	FParse::Value(CmdLine, TEXT("UBASentryUrl="), UbaSentryUrl);
+	CurrentSource = TEXT("CMD");
+
+	ParseCmdLineString(CmdLine, TEXT("UBAHorde="), Server);
+	ParseCmdLineString(CmdLine, TEXT("UBAHordeToken="), Token);
+	ParseCmdLineString(CmdLine, TEXT("UBAHordePool="), Pool);
+	ParseCmdLineBoolean(CmdLine, TEXT("UBAHordeAllowWine"), bAllowWine);
+	ParseCmdLineInt(CmdLine, TEXT("UBAHordeMaxCores="), MaxCores);
+	ParseCmdLineString(CmdLine, TEXT("UBAHordeHost="), Host);
+	ParseCmdLineString(CmdLine, TEXT("UBAHordeCluster="), Cluster);
+	ParseCmdLineString(CmdLine, TEXT("UBAHordeConnectionMode="), ConnectionMode);
+	ParseCmdLineString(CmdLine, TEXT("UBAHordeEncryption="), Encryption);
+	ParseCmdLineString(CmdLine, TEXT("UBASentryUrl="), SentryUrl);
 	if (FParse::Param(CmdLine, TEXT("UBADisableHorde")))
 	{
-		HordeEnabled = EUbaHordeEnabledMode::Disabled;
+		ProviderEnabled.Value = TEXT("False");
+		ProviderEnabled.Source = CurrentSource;
 	}
 	else if (FParse::Param(CmdLine, TEXT("UBAEnableHorde")))
 	{
-		HordeEnabled = EUbaHordeEnabledMode::Enabled;
+		ProviderEnabled.Value = TEXT("True");
+		ProviderEnabled.Source = CurrentSource;
 	}
 
-	FParse::Value(CmdLine, TEXT("MaxParallelActions="), MaxParallelActions);
+	ParseCmdLineInt(CmdLine, TEXT("MaxParallelActions="), MaxParallelActions);
 	if (FParse::Param(CmdLine, TEXT("ExclusiveRemoteShaderCompiling")))
 	{
-		MaxParallelActions = 0; // -ExclusiveRemoteShaderCompiling overrides -MaxParallelActions
+		MaxParallelActions.Value = 0; // -ExclusiveRemoteShaderCompiling overrides -MaxParallelActions
+		MaxParallelActions.Source = CurrentSource;
 	}
 }
 
@@ -280,6 +320,13 @@ static TArray<FString> GetBuildConfigFilePaths()
 	return OutPaths;
 }
 
+static bool IsUbaProviderEnabled(const TCHAR* InEnabledState)
+{
+	EUbaHordeEnabledMode HordeEnabledMode = EUbaHordeEnabledMode::Disabled;
+	LexFromString(HordeEnabledMode, InEnabledState);
+	return HordeEnabledMode == EUbaHordeEnabledMode::Enabled || (HordeEnabledMode == EUbaHordeEnabledMode::EnabledForBuildMachineOnly && GIsBuildMachine);
+}
+
 /**
  * Reads the UBA / Horde configuration in the following order :
  *  1. From a list of BuildConfiguration.xml files (in application settings and user folders)
@@ -288,20 +335,45 @@ static TArray<FString> GetBuildConfigFilePaths()
  */
 void FUbaHordeConfig::Initialize()
 {
+	FUbaHordeConfigParser ConfigParser;
+
 	// Read configurations from BuildConfiguration.xml files
 	for (const FString& ConfigFilePath : GetBuildConfigFilePaths())
 	{
-		if (ParseConfigFromXmlFile(*ConfigFilePath))
+		if (ConfigParser.ParseConfigFromXmlFile(*ConfigFilePath))
 		{
 			UE_LOG(LogUbaHorde, Verbose, TEXT("Read UBA/Horde settings from %s"), *ConfigFilePath);
 		}
 	}
 
 	// Override configuration entries by their configuration INI files
-	ParseConfigFromIniFile();
+	ConfigParser.ParseConfigFromIniFile();
 
 	// Override configuration entries by their command line arguments
-	ParseConfigFromCommandline(FCommandLine::Get());
+	ConfigParser.ParseConfigFromCommandline(FCommandLine::Get());
+
+	// Convert parsed argumetns to public configuration parameters
+	ConfigParser.GetConfig(ConfigParser.Provider, Provider);
+	bIsProviderEnabled = IsUbaProviderEnabled(*ConfigParser.ProviderEnabled.Value);
+	ConfigParser.GetConfig(ConfigParser.Pool, HordePool);
+	ConfigParser.GetConfig(ConfigParser.Server, HordeServer);
+	ConfigParser.GetConfig(ConfigParser.Token, HordeToken);
+	ConfigParser.GetConfig(ConfigParser.Condition, HordeCondition);
+	ConfigParser.GetConfig(ConfigParser.Cluster, HordeCluster);
+	ConfigParser.GetConfig(ConfigParser.Host, HordeHost);
+	ConfigParser.GetConfig(ConfigParser.SentryUrl, UbaSentryUrl);
+	ConfigParser.GetConfig(ConfigParser.MaxCores, HordeMaxCores);
+	ConfigParser.GetConfig(ConfigParser.MaxParallelActions, MaxParallelActions);
+	ConfigParser.GetConfig(ConfigParser.bAllowWine, bHordeAllowWine);
+
+	if (ConfigParser.ConnectionMode.Source)
+	{
+		LexFromString(HordeConnectionMode, *ConfigParser.ConnectionMode.Value);
+	}
+	if (ConfigParser.Encryption.Source)
+	{
+		LexFromString(HordeEncryption, *ConfigParser.Encryption.Value);
+	}
 
 	// If connection mode is "realy"-mode, encryption is implied and we use AES since UbaController does not support SSL/TLS yet
 	if (HordeConnectionMode == EUbaHordeConnectionMode::Relay && HordeEncryption == EUbaHordeEncryption::None)
@@ -309,40 +381,46 @@ void FUbaHordeConfig::Initialize()
 		HordeEncryption = EUbaHordeEncryption::AES;
 	}
 
-	if (IsHordeEnabled())
+	if (bIsProviderEnabled)
 	{
+		auto GetConfigSource = [](const TCHAR* InSource) -> const TCHAR*
+			{
+				return InSource ? InSource : TEXT("Default");
+			};
+
 		UE_LOG(
 			LogUbaHorde, Display,
-			TEXT("UBA/Horde Configuration:")
-			TEXT("\n  Server:            %s")
-			TEXT("\n  Pool:              %s")
-			TEXT("\n  Condition:         %s")
-			TEXT("\n  Cluster:           %s")
-			TEXT("\n  Host:              %s")
-			TEXT("\n  MaxCores:          %d")
-			TEXT("\n  Allow Wine:        %s")
-			TEXT("\n  Mode:              %s")
-			TEXT("\n  Encryption:        %s")
-			, *HordeServer
-			, *GetHordePool()
-			, *HordeCondition
-			, *HordeCluster
-			, *HordeHost
-			, HordeMaxCores
-			, bHordeAllowWine ? TEXT("Yes") : TEXT("No")
-			, LexToString(HordeConnectionMode)
-			, LexToString(HordeEncryption)
+			TEXT("UBA/Horde Configuration [%s]:")
+			TEXT("\n  Server:            %s (%s)")
+			TEXT("\n  Pool:              %s (%s)")
+			TEXT("\n  Condition:         %s (%s)")
+			TEXT("\n  Cluster:           %s (%s)")
+			TEXT("\n  Host:              %s (%s)")
+			TEXT("\n  MaxCores:          %d (%s)")
+			TEXT("\n  Allow Wine:        %s (%s)")
+			TEXT("\n  Mode:              %s (%s)")
+			TEXT("\n  Encryption:        %s (%s)")
+			, *Provider
+			, *HordeServer, GetConfigSource(ConfigParser.Server.Source)
+			, *GetHordePool(), GetConfigSource(ConfigParser.Pool.Source)
+			, *HordeCondition, GetConfigSource(ConfigParser.Condition.Source)
+			, *HordeCluster, GetConfigSource(ConfigParser.Cluster.Source)
+			, *HordeHost, GetConfigSource(ConfigParser.Host.Source)
+			, HordeMaxCores, GetConfigSource(ConfigParser.MaxCores.Source)
+			, bHordeAllowWine ? TEXT("Yes") : TEXT("No"), GetConfigSource(ConfigParser.bAllowWine.Source)
+			, LexToString(HordeConnectionMode), GetConfigSource(ConfigParser.ConnectionMode.Source)
+			, LexToString(HordeEncryption), GetConfigSource(ConfigParser.Encryption.Source)
 		);
 	}
 	else
 	{
-		UE_LOG(LogUbaHorde, Display, TEXT("UBA/Horde Configuration: Not Enabled"));
+		UE_LOG(LogUbaHorde, Display, TEXT("UBA/Horde Configuration [%s]: Not Enabled"), *Provider);
 	}
 }
 
 UBACOORDINATORHORDE_API bool FUbaHordeConfig::IsHordeEnabled() const
 {
-	return HordeEnabled == EUbaHordeEnabledMode::Enabled || (HordeEnabled == EUbaHordeEnabledMode::EnabledForBuildMachineOnly && GIsBuildMachine);
+	return bIsProviderEnabled;
 }
 
 UBACOORDINATORHORDE_API const FUbaHordeConfig& FUbaHordeConfig::Get()
diff --git a/Engine/Source/Developer/UbaCoordinatorHorde/Public/UbaHordeConfig.h b/Engine/Source/Developer/UbaCoordinatorHorde/Public/UbaHordeConfig.h
--- a/Engine/Source/Developer/UbaCoordinatorHorde/Public/UbaHordeConfig.h
+++ b/Engine/Source/Developer/UbaCoordinatorHorde/Public/UbaHordeConfig.h
@@ -35,27 +35,21 @@ public:
 	static const TCHAR* ClusterDefault;
 	static const TCHAR* ClusterAuto;
 
+	/* Selected Horde provider, e.g. "Uba.Provider.Horde" by default */
+	FString Provider;
+
+	/** Specifies whether UBA/Horde is enabled for the current session. By default disabled. */
+	bool bIsProviderEnabled = false;
+
+	/** Horde pool name. */
+	FString HordePool;
+
 	/** URI of the Horde server. */
 	FString HordeServer;
 
 	/** Authentication token for the Horde server. */
 	FString HordeToken;
 
-	/** Default Horde pool name. */
-	FString DefaultHordePool;
-
-	/** Overrides the default pool name when no platform specific pool is specified. */
-	FString OverrideHordePool;
-
-	/** Horde pool name for Windows target platform. */
-	FString WindowsHordePool;
-
-	/** Horde pool name for Mac target platform. */
-	FString MacHordePool;
-
-	/** Horde pool name for Linux target platform. */
-	FString LinuxHordePool;
-
 	/** Requirements for the Horde agent to assign. */
 	FString HordeCondition;
 
@@ -83,25 +77,38 @@ public:
 	/** Specifies transport layer encryption. Currently only AES encryption is supported. */
 	EUbaHordeEncryption HordeEncryption = EUbaHordeEncryption::None;
 
-	/** Specifies whether UBA/Horde is enabled for the current session. By default disabled. */
-	EUbaHordeEnabledMode HordeEnabled = EUbaHordeEnabledMode::Disabled;
-	
 	/** Returns the active Horde pool depending on override settings and current desktop platform. */
 	UBACOORDINATORHORDE_API const FString& GetHordePool() const;
 
-	/** Returns true if UBA/Horde is enabled. This depends on HordeEnabled and whether the process is running on a build machine. */
-	UBACOORDINATORHORDE_API bool IsHordeEnabled() const;
-
 	/** Returns the global UBA/Horde configuration and initializes it from a list of XML files on the first call. */
 	UBACOORDINATORHORDE_API static const FUbaHordeConfig& Get();
 
+	// DEPRECATED
+
+	UE_DEPRECATED(5.7, "Use HordePool instead.")
+	FString DefaultHordePool;
+
+	UE_DEPRECATED(5.7, "Use HordePool instead.")
+	FString OverrideHordePool;
+
+	UE_DEPRECATED(5.7, "Use HordePool instead. Windows specific pool can be set via [Uba.Provider.Horde]:Pool in Config/Win64/BaseEngine.ini")
+	FString WindowsHordePool;
+
+	UE_DEPRECATED(5.7, "Use HordePool instead. Mac specific pool can be set via [Uba.Provider.Horde]:Pool in Config/Mac/BaseEngine.ini")
+	FString MacHordePool;
+
+	UE_DEPRECATED(5.7, "Use HordePool instead. Linux specific pool can be set via [Uba.Provider.Horde]:Pool in Config/Linux/BaseEngine.ini")
+	FString LinuxHordePool;
+
+	UE_DEPRECATED(5.7, "Use bIsProviderEnabled instead")
+	EUbaHordeEnabledMode HordeEnabled = EUbaHordeEnabledMode::Disabled;
+
+	UE_DEPRECATED(5.7, "Use bIsProviderEnabled instead")
+	UBACOORDINATORHORDE_API bool IsHordeEnabled() const;
+
 private:
 	void Initialize();
 
-	bool ParseConfigFromXmlFile(const TCHAR* InBuildConfigFilename);
-	void ParseConfigFromIniFile();
-	void ParseConfigFromCommandline(const TCHAR* CmdLine);
-
 };
 
 UBACOORDINATORHORDE_API const TCHAR* LexToString(EUbaHordeConnectionMode InConnectionMode);
