[UBA] Change deprecation version in FUbaHordeConfig from 5.7 to 5.6 after integration into 5.6 release stream.

diff --git a/Engine/Source/Developer/UbaCoordinatorHorde/Private/UbaHordeConfig.cpp b/Engine/Source/Developer/UbaCoordinatorHorde/Private/UbaHordeConfig.cpp
--- a/Engine/Source/Developer/UbaCoordinatorHorde/Private/UbaHordeConfig.cpp
+++ b/Engine/Source/Developer/UbaCoordinatorHorde/Private/UbaHordeConfig.cpp
@@ -222,7 +222,7 @@ void FUbaHordeConfigParser::ParseConfigFromIniFile()
 	// Also check [UbaController]:Enabled for compatibility
 	if (GetUbaConfigString(TEXT("UbaController"), TEXT("Enabled"), ProviderEnabled))
 	{
-		UE_LOG(LogUbaHorde, Warning, TEXT("Configuration '[UbaController]:Enabled' is deprecated since UE 5.7; Use '[%s]:Enabled' instead"), *Provider.Value);
+		UE_LOG(LogUbaHorde, Warning, TEXT("Configuration '[UbaController]:Enabled' is deprecated since UE 5.6; Use '[%s]:Enabled' instead"), *Provider.Value);
 	}
 
 	// Read all configurations from provider section
diff --git a/Engine/Source/Developer/UbaCoordinatorHorde/Public/UbaHordeConfig.h b/Engine/Source/Developer/UbaCoordinatorHorde/Public/UbaHordeConfig.h
--- a/Engine/Source/Developer/UbaCoordinatorHorde/Public/UbaHordeConfig.h
+++ b/Engine/Source/Developer/UbaCoordinatorHorde/Public/UbaHordeConfig.h
@@ -85,25 +85,25 @@ public:
 
 	// DEPRECATED
 
-	UE_DEPRECATED(5.7, "Use HordePool instead.")
+	UE_DEPRECATED(5.6, "Use HordePool instead.")
 	FString DefaultHordePool;
 
-	UE_DEPRECATED(5.7, "Use HordePool instead.")
+	UE_DEPRECATED(5.6, "Use HordePool instead.")
 	FString OverrideHordePool;
 
-	UE_DEPRECATED(5.7, "Use HordePool instead. Windows specific pool can be set via [Uba.Provider.Horde]:Pool in Config/Win64/BaseEngine.ini")
+	UE_DEPRECATED(5.6, "Use HordePool instead. Windows specific pool can be set via [Uba.Provider.Horde]:Pool in Config/Win64/BaseEngine.ini")
 	FString WindowsHordePool;
 
-	UE_DEPRECATED(5.7, "Use HordePool instead. Mac specific pool can be set via [Uba.Provider.Horde]:Pool in Config/Mac/BaseEngine.ini")
+	UE_DEPRECATED(5.6, "Use HordePool instead. Mac specific pool can be set via [Uba.Provider.Horde]:Pool in Config/Mac/BaseEngine.ini")
 	FString MacHordePool;
 
-	UE_DEPRECATED(5.7, "Use HordePool instead. Linux specific pool can be set via [Uba.Provider.Horde]:Pool in Config/Linux/BaseEngine.ini")
+	UE_DEPRECATED(5.6, "Use HordePool instead. Linux specific pool can be set via [Uba.Provider.Horde]:Pool in Config/Linux/BaseEngine.ini")
 	FString LinuxHordePool;
 
-	UE_DEPRECATED(5.7, "Use bIsProviderEnabled instead")
+	UE_DEPRECATED(5.6, "Use bIsProviderEnabled instead")
 	EUbaHordeEnabledMode HordeEnabled = EUbaHordeEnabledMode::Disabled;
 
-	UE_DEPRECATED(5.7, "Use bIsProviderEnabled instead")
+	UE_DEPRECATED(5.6, "Use bIsProviderEnabled instead")
 	UBACOORDINATORHORDE_API bool IsHordeEnabled() const;
 
 private:
