[UBA] Include host OS family as condition for UBA remote agent requests.

This adopts the same conditions of the host OS in UbaController from UBT. It also makes use of the "AllowWine" property of the respective Horde provider. The new Json body for the agent requirements can now look like this on Windows:
"requirements":
{
"condition": "(OSFamily == 'Windows' || WineEnabled == 'true')",
"exclusive": true
}

#rnx

diff --git a/Engine/Source/Developer/UbaCoordinatorHorde/Private/UbaHordeAgentManager.cpp b/Engine/Source/Developer/UbaCoordinatorHorde/Private/UbaHordeAgentManager.cpp
--- a/Engine/Source/Developer/UbaCoordinatorHorde/Private/UbaHordeAgentManager.cpp
+++ b/Engine/Source/Developer/UbaCoordinatorHorde/Private/UbaHordeAgentManager.cpp
@@ -161,8 +161,10 @@ static FDirectoryEntry CreateHordeBundleDirectoryEntry(const std::filesystem::pa
 	return FDirectoryEntry(DirectoryHandle, FIoHash(), FUtf8String(Path.filename().string().c_str()), Length);
 }
 
-bool CreateHordeBundleFromFile(const std::filesystem::path& InputFilename, const std::filesystem::path& OutputFilename)
+static bool CreateHordeBundleFromFile(const std::filesystem::path& InputFilename, const std::filesystem::path& OutputFilename)
 {
+	check(!InputFilename.empty());
+
 	TSharedRef<FFileStorageClient> FileStorage = MakeShared<FFileStorageClient>(OutputFilename.parent_path());
 	TSharedRef<FBundleStorageClient> Storage = MakeShared<FBundleStorageClient>(FileStorage);
 
@@ -342,7 +344,15 @@ void FUbaHordeAgentManager::ThreadAgent(FHordeAgentWrapper& Wrapper)
 
 		// Initialize horde request with connection mode preference and encryption.
 		// Pool ID is superseded by cluster ID, but it still has to be provided if default cluster and direct mode is used.
-		const FString HordeRequestJsonBody = HordeMetaClient->BuildHordeRequestJsonBody(Config.HordeConnectionMode == EUbaHordeConnectionMode::Direct ? Config.GetHordePool() : TEXT(""), Config.HordeConnectionMode, Config.HordeEncryption);
+		constexpr bool bExclusiveAccess = true;
+		const FString HordeRequestJsonBody = HordeMetaClient->BuildHordeRequestJsonBody(
+			Config.HordeConnectionMode == EUbaHordeConnectionMode::Direct ? Config.GetHordePool() : TEXT(""),
+			Config.HordeConnectionMode,
+			Config.HordeEncryption,
+			*Config.HordeCondition,
+			bExclusiveAccess,
+			Config.bHordeAllowWine
+		);
 
 		// Resolve cluster ID through Horde if "_auto" was specified
 		FHordeClusterInfo ClusterInfo;
diff --git a/Engine/Source/Developer/UbaCoordinatorHorde/Private/UbaHordeConfig.cpp b/Engine/Source/Developer/UbaCoordinatorHorde/Private/UbaHordeConfig.cpp
--- a/Engine/Source/Developer/UbaCoordinatorHorde/Private/UbaHordeConfig.cpp
+++ b/Engine/Source/Developer/UbaCoordinatorHorde/Private/UbaHordeConfig.cpp
@@ -249,6 +249,7 @@ void FUbaHordeConfigParser::ParseConfigFromCommandline(const TCHAR* CmdLine)
 	ParseCmdLineInt(CmdLine, TEXT("UBAHordeMaxCores="), MaxCores);
 	ParseCmdLineString(CmdLine, TEXT("UBAHordeHost="), Host);
 	ParseCmdLineString(CmdLine, TEXT("UBAHordeCluster="), Cluster);
+	ParseCmdLineString(CmdLine, TEXT("UBAHordeRequirements="), Condition);
 	ParseCmdLineString(CmdLine, TEXT("UBAHordeConnectionMode="), ConnectionMode);
 	ParseCmdLineString(CmdLine, TEXT("UBAHordeEncryption="), Encryption);
 	ParseCmdLineString(CmdLine, TEXT("UBASentryUrl="), SentryUrl);
diff --git a/Engine/Source/Developer/UbaCoordinatorHorde/Private/UbaHordeMetaClient.cpp b/Engine/Source/Developer/UbaCoordinatorHorde/Private/UbaHordeMetaClient.cpp
--- a/Engine/Source/Developer/UbaCoordinatorHorde/Private/UbaHordeMetaClient.cpp
+++ b/Engine/Source/Developer/UbaCoordinatorHorde/Private/UbaHordeMetaClient.cpp
@@ -39,7 +39,7 @@ bool FUbaHordeMetaClient::RefreshHttpClient()
 	return true;
 }
 
-FString FUbaHordeMetaClient::BuildHordeRequestJsonBody(const FString& PoolId, EUbaHordeConnectionMode ConnectionMode, EUbaHordeEncryption Encryption, bool bExclusiveAccess)
+FString FUbaHordeMetaClient::BuildHordeRequestJsonBody(const FString& PoolId, EUbaHordeConnectionMode ConnectionMode, EUbaHordeEncryption Encryption, const TCHAR* ConditionSuffix, bool bExclusiveAccess, bool bAllowWine)
 {
 	// Construct JSON body for Horde request, e.g. '{ "connection": { "modePreference": "relay", "encryption": "aes" } }'
 	const TSharedPtr<FJsonObject> RequestJNode = MakeShared<FJsonObject>();
@@ -52,6 +52,21 @@ FString FUbaHordeMetaClient::BuildHordeRequestJsonBody(const FString& PoolId, EU
 				{
 					RequirementsJNode->SetStringField(TEXT("pool"), PoolId);
 				}
+
+				FString Condition;
+#if PLATFORM_WINDOWS
+				Condition.Appendf(TEXT("(OSFamily == 'Windows' || WineEnabled == '%s')"), bAllowWine ? TEXT("true") : TEXT("false"));
+#elif PLATFORM_MAC
+				Condition.Append(TEXT("OSFamily == 'MacOS'"));
+#elif PLATFORM_LINUX
+				Condition.Append(TEXT("OSFamily == 'Linux'"));
+#endif
+				if (ConditionSuffix && *ConditionSuffix != TEXT('\0'))
+				{
+					Condition.Appendf(TEXT(" %s"), ConditionSuffix);
+				}
+				RequirementsJNode->SetStringField(TEXT("condition"), Condition);
+
 				RequirementsJNode->SetBoolField(TEXT("exclusive"), bExclusiveAccess);
 			}
 			RequestJNode->SetObjectField(TEXT("requirements"), RequirementsJNode);
diff --git a/Engine/Source/Developer/UbaCoordinatorHorde/Private/UbaHordeMetaClient.h b/Engine/Source/Developer/UbaCoordinatorHorde/Private/UbaHordeMetaClient.h
--- a/Engine/Source/Developer/UbaCoordinatorHorde/Private/UbaHordeMetaClient.h
+++ b/Engine/Source/Developer/UbaCoordinatorHorde/Private/UbaHordeMetaClient.h
@@ -82,7 +82,7 @@ public:
 
 	bool RefreshHttpClient();
 
-	static FString BuildHordeRequestJsonBody(const FString& PoolId, EUbaHordeConnectionMode ConnectionMode, EUbaHordeEncryption Encryption, bool bExclusiveAccess = true);
+	static FString BuildHordeRequestJsonBody(const FString& PoolId, EUbaHordeConnectionMode ConnectionMode, EUbaHordeEncryption Encryption, const TCHAR* ConditionSuffix = nullptr, bool bExclusiveAccess = true, bool bAllowWine = true);
 
 	TSharedPtr<HordeClusterPromise, ESPMode::ThreadSafe> RequestClusterId(const FString& HordeRequestJsonBody);
 
