[UbaCoordinatorHorde]
* Added so -killtcphogs is added to cmd line on mac
* Added so UbaWine.dll.so is transferred to horde agent
* Set initial value of LastRequestFailTime to 0 to have visualizer show correct stats.

diff --git a/Engine/Source/Developer/UbaCoordinatorHorde/Private/UbaHordeAgentManager.cpp b/Engine/Source/Developer/UbaCoordinatorHorde/Private/UbaHordeAgentManager.cpp
--- a/Engine/Source/Developer/UbaCoordinatorHorde/Private/UbaHordeAgentManager.cpp
+++ b/Engine/Source/Developer/UbaCoordinatorHorde/Private/UbaHordeAgentManager.cpp
@@ -39,7 +39,7 @@ namespace UbaCoordinatorHordeModule
 FUbaHordeAgentManager::FUbaHordeAgentManager(const FString& InWorkingDir, const FString& InBinariesPath)
 	:	WorkingDir(InWorkingDir)
 	,	BinariesPath(InBinariesPath)
-	,	LastRequestFailTime(1)
+	,	LastRequestFailTime(0)
 	,	TargetCoreCount(0)
 	,	EstimatedCoreCount(0)
 	,   ActiveCoreCount(0)
@@ -268,6 +268,7 @@ void FUbaHordeAgentManager::ThreadAgent(FHordeAgentWrapper& Wrapper)
 			{
 #if PLATFORM_WINDOWS
 				{ TEXT("UbaAgent.exe"), TEXT("UbaAgent.Bundle.ref") },
+				{ TEXT("UbaWine.dll.so"), TEXT("UbaWine.Bundle.ref") },
 #elif PLATFORM_LINUX
 				{ TEXT("UbaAgent"), TEXT("UbaAgent.Bundle.ref") },
 				{ TEXT("UbaAgent.debug"), TEXT("UbaAgent.debug.Bundle.ref") },
@@ -478,6 +479,10 @@ void FUbaHordeAgentManager::ThreadAgent(FHordeAgentWrapper& Wrapper)
 		// After 15 seconds of idling agent will automatically disconnect
 		UbaAgentArgs.AddRaw("-MaxIdle=15");
 
+		#if PLATFORM_MAC
+		UbaAgentArgs.Add("-KillTcpHogs");
+		#endif
+
 		UbaAgentArgs.AddRaw("-Dir=%UE_HORDE_SHARED_DIR%\\Uba");
 		UbaAgentArgs.AddRaw("-Eventfile=%UE_HORDE_TERMINATION_SIGNAL_FILE%");
 		UbaAgentArgs.Add(FAnsiString::Printf("-Description=%s", TCHAR_TO_ANSI(*Agent->GetMachineInfo().LeaseLink)));
