[UBA] Explicitly show what UBA/Horde config properties are unspecified via "<None>" string.

#rnx

diff --git a/Engine/Source/Developer/UbaCoordinatorHorde/Private/UbaHordeConfig.cpp b/Engine/Source/Developer/UbaCoordinatorHorde/Private/UbaHordeConfig.cpp
--- a/Engine/Source/Developer/UbaCoordinatorHorde/Private/UbaHordeConfig.cpp
+++ b/Engine/Source/Developer/UbaCoordinatorHorde/Private/UbaHordeConfig.cpp
@@ -376,17 +376,25 @@ void FUbaHordeConfig::Initialize()
 		LexFromString(HordeEncryption, *ConfigParser.Encryption.Value);
 	}
 
+	constexpr const TCHAR* kConfigSourceImplied = TEXT("Implied");
+
 	// If connection mode is "realy"-mode, encryption is implied and we use AES since UbaController does not support SSL/TLS yet
 	if (HordeConnectionMode == EUbaHordeConnectionMode::Relay && HordeEncryption == EUbaHordeEncryption::None)
 	{
 		HordeEncryption = EUbaHordeEncryption::AES;
+		ConfigParser.Encryption.Source = kConfigSourceImplied;
 	}
 
 	if (bIsProviderEnabled)
 	{
+		auto GetConfigValue = [](const TCHAR* InValue) -> const TCHAR*
+			{
+				return InValue && *InValue != TEXT('\0') ? InValue : TEXT("<None>");
+			};
+
 		auto GetConfigSource = [](const TCHAR* InSource) -> const TCHAR*
 			{
-				return InSource ? InSource : TEXT("Default");
+				return InSource && *InSource != TEXT('\0') ? InSource : TEXT("Default");
 			};
 
 		UE_LOG(
@@ -402,15 +410,15 @@ void FUbaHordeConfig::Initialize()
 			TEXT("\n  Mode:              %s (%s)")
 			TEXT("\n  Encryption:        %s (%s)")
 			, *Provider
-			, *HordeServer, GetConfigSource(ConfigParser.Server.Source)
-			, *GetHordePool(), GetConfigSource(ConfigParser.Pool.Source)
-			, *HordeCondition, GetConfigSource(ConfigParser.Condition.Source)
-			, *HordeCluster, GetConfigSource(ConfigParser.Cluster.Source)
-			, *HordeHost, GetConfigSource(ConfigParser.Host.Source)
+			, GetConfigValue(*HordeServer), GetConfigSource(ConfigParser.Server.Source)
+			, GetConfigValue(*GetHordePool()), GetConfigSource(ConfigParser.Pool.Source)
+			, GetConfigValue(*HordeCondition), GetConfigSource(ConfigParser.Condition.Source)
+			, GetConfigValue(*HordeCluster), GetConfigSource(ConfigParser.Cluster.Source)
+			, GetConfigValue(*HordeHost), GetConfigSource(ConfigParser.Host.Source)
 			, HordeMaxCores, GetConfigSource(ConfigParser.MaxCores.Source)
 			, bHordeAllowWine ? TEXT("Yes") : TEXT("No"), GetConfigSource(ConfigParser.bAllowWine.Source)
-			, LexToString(HordeConnectionMode), GetConfigSource(ConfigParser.ConnectionMode.Source)
-			, LexToString(HordeEncryption), GetConfigSource(ConfigParser.Encryption.Source)
+			, GetConfigValue(LexToString(HordeConnectionMode)), GetConfigSource(ConfigParser.ConnectionMode.Source)
+			, GetConfigValue(LexToString(HordeEncryption)), GetConfigSource(ConfigParser.Encryption.Source)
 		);
 	}
 	else
