[UbaCoordinatorHorde]
* Fixed so horde agent manager is not counting number of log entries but instead actually waits for the "Listening on" log entry

diff --git a/Engine/Source/Developer/UbaCoordinatorHorde/Private/UbaHordeAgent.cpp b/Engine/Source/Developer/UbaCoordinatorHorde/Private/UbaHordeAgent.cpp
--- a/Engine/Source/Developer/UbaCoordinatorHorde/Private/UbaHordeAgent.cpp
+++ b/Engine/Source/Developer/UbaCoordinatorHorde/Private/UbaHordeAgent.cpp
@@ -224,7 +224,7 @@ bool FUbaHordeAgent::IsValid()
 	return bIsValid && !bHasErrors;
 }
 
-void FUbaHordeAgent::Poll(bool LogReports)
+void FUbaHordeAgent::Poll(bool LogReports, FString* OutResponseText)
 {
 	EAgentMessageType Type = EAgentMessageType::None;
 	constexpr int32 ReadresponseTimeoutMS = 100;
@@ -253,6 +253,12 @@ void FUbaHordeAgent::Poll(bool LogReports)
 				{
 					UE_LOG(LogUbaHordeAgent, Log, TEXT("Response [ExecuteOutput]: %s"), ResponseString.Get());
 				}
+
+				if (OutResponseText)
+				{
+					*OutResponseText = MoveTemp(ResponseString);
+					return;
+				}
 			}
 		}
 		break;
diff --git a/Engine/Source/Developer/UbaCoordinatorHorde/Private/UbaHordeAgent.h b/Engine/Source/Developer/UbaCoordinatorHorde/Private/UbaHordeAgent.h
--- a/Engine/Source/Developer/UbaCoordinatorHorde/Private/UbaHordeAgent.h
+++ b/Engine/Source/Developer/UbaCoordinatorHorde/Private/UbaHordeAgent.h
@@ -31,7 +31,7 @@ public:
 	bool IsValid();
 
 	// Reads output from the child channel and reports them to the log of the calling process.
-	void Poll(bool LogReports);
+	void Poll(bool LogReports, FString* OutResponseText = nullptr);
 
 	inline const FHordeRemoteMachineInfo& GetMachineInfo() const
 	{
diff --git a/Engine/Source/Developer/UbaCoordinatorHorde/Private/UbaHordeAgentManager.cpp b/Engine/Source/Developer/UbaCoordinatorHorde/Private/UbaHordeAgentManager.cpp
--- a/Engine/Source/Developer/UbaCoordinatorHorde/Private/UbaHordeAgentManager.cpp
+++ b/Engine/Source/Developer/UbaCoordinatorHorde/Private/UbaHordeAgentManager.cpp
@@ -523,45 +523,51 @@ void FUbaHordeAgentManager::ThreadAgent(FHordeAgentWrapper& Wrapper)
 
 	ON_SCOPE_EXIT{ EstimatedCoreCount -= MachineCoreCount; ActiveCoreCount -= MachineCoreCount; --AgentsActive; UpdateStatus(nullptr); };
 
-	uint32 callCounter = 0; // TODO: This should react on the listen string instead of waiting for two text messages :)
+	FString Response;
+	bool bConnected = false;
 
 	while (Agent->IsValid() && !ShouldExit.Wait(100))
 	{
-		Agent->Poll(UbaCoordinatorHordeModule::bHordeForwardAgentLogs);
+		bool bReadResponse = (bUseListen && !bConnected);
+		Agent->Poll(bReadResponse || UbaCoordinatorHordeModule::bHordeForwardAgentLogs, bReadResponse ? &Response : nullptr);
 
 		if (!bUseListen)
 		{
 			continue;
 		}
 
-		if (callCounter++ == 2)
+		if (!Response.Contains(TEXT("Listening on")))
 		{
-			FString CryptoNonce16;
-			if (Agent->GetMachineInfo().Encryption != EUbaHordeEncryption::None)
+			continue;
+		}
+
+		FString CryptoNonce16;
+		if (Agent->GetMachineInfo().Encryption != EUbaHordeEncryption::None)
+		{
+			// Generate random crypto nonce for UbaAgent
+			CryptoNonce16.Reserve(16);
+			for (int32 CharacterIndex = 0; CharacterIndex < 16; ++CharacterIndex)
 			{
-				// Generate random crypto nonce for UbaAgent
-				CryptoNonce16.Reserve(16);
-				for (int32 CharacterIndex = 0; CharacterIndex < 16; ++CharacterIndex)
-				{
-					CryptoNonce16.AppendChar(TEXT("0123456789abcdef")[FMath::Rand() % 0xF]);
-				}
+				CryptoNonce16.AppendChar(TEXT("0123456789abcdef")[FMath::Rand() % 0xF]);
 			}
+		}
 
-			// Add this machine as client to the remote agent.
-			// Use JSON field "port" when adding clients while "agentPort" is used for UbaAgent to listen to on the server side.
-			const bool bAddClientSuccess = AddClientCallback(
-				AddClientUserData,
-				StringCast<uba::tchar>(*UbaAgentIp).Get(),
-				static_cast<uint16>(UbaPort.Port),
-				Agent->GetMachineInfo().Encryption != EUbaHordeEncryption::None ? StringCast<uba::tchar>(*CryptoNonce16).Get() : nullptr
-			);
+		// Add this machine as client to the remote agent.
+		// Use JSON field "port" when adding clients while "agentPort" is used for UbaAgent to listen to on the server side.
+		const bool bAddClientSuccess = AddClientCallback(
+			AddClientUserData,
+			StringCast<uba::tchar>(*UbaAgentIp).Get(),
+			static_cast<uint16>(UbaPort.Port),
+			Agent->GetMachineInfo().Encryption != EUbaHordeEncryption::None ? StringCast<uba::tchar>(*CryptoNonce16).Get() : nullptr
+		);
 
-			if (!bAddClientSuccess)
-			{
-				UE_LOG(LogUbaHorde, Display, TEXT("Server_AddClient(%s:%u) failed"), *UbaAgentIp, UbaPort.Port);
-				return;
-			}
+		if (!bAddClientSuccess)
+		{
+			UE_LOG(LogUbaHorde, Display, TEXT("Server_AddClient(%s:%u) failed"), *UbaAgentIp, UbaPort.Port);
+			return;
 		}
+
+		bConnected = true;
 	}
 }
 
