[UbaCoordinatorHorde]
* Fixed bug where host ended up creating lots of connections to agent

diff --git a/Engine/Source/Developer/UbaCoordinatorHorde/Private/UbaHordeAgentManager.cpp b/Engine/Source/Developer/UbaCoordinatorHorde/Private/UbaHordeAgentManager.cpp
--- a/Engine/Source/Developer/UbaCoordinatorHorde/Private/UbaHordeAgentManager.cpp
+++ b/Engine/Source/Developer/UbaCoordinatorHorde/Private/UbaHordeAgentManager.cpp
@@ -531,7 +531,7 @@ void FUbaHordeAgentManager::ThreadAgent(FHordeAgentWrapper& Wrapper)
 		bool bReadResponse = (bUseListen && !bConnected);
 		Agent->Poll(bReadResponse || UbaCoordinatorHordeModule::bHordeForwardAgentLogs, bReadResponse ? &Response : nullptr);
 
-		if (!bUseListen)
+		if (!bUseListen || bConnected)
 		{
 			continue;
 		}
