[UBA] Minor improvement of deprecation warning when [UbaController]:Enabled is still used.

#rnx

diff --git a/Engine/Source/Developer/UbaCoordinatorHorde/Private/UbaHordeConfig.cpp b/Engine/Source/Developer/UbaCoordinatorHorde/Private/UbaHordeConfig.cpp
--- a/Engine/Source/Developer/UbaCoordinatorHorde/Private/UbaHordeConfig.cpp
+++ b/Engine/Source/Developer/UbaCoordinatorHorde/Private/UbaHordeConfig.cpp
@@ -222,7 +222,7 @@ void FUbaHordeConfigParser::ParseConfigFromIniFile()
 	// Also check [UbaController]:Enabled for compatibility
 	if (GetUbaConfigString(TEXT("UbaController"), TEXT("Enabled"), ProviderEnabled))
 	{
-		UE_LOG(LogUbaHorde, Warning, TEXT("Configuration '[UbaController]:Enabled' is deprecated since UE 5.6; Use '[%s]:Enabled' instead"), *Provider.Value);
+		UE_LOG(LogUbaHorde, Warning, TEXT("Configuration '[UbaController]:Enabled' is deprecated since UE 5.6; Add an enabled Horde provider to '[UbaController]:+Providers' instead"));
 	}
 
 	// Read all configurations from provider section
