[UBA] Simplify enabling UBA provider for shader compilation via "-UBARelay" argument.

This adds two new UBA command line arguments to the UbaController plugin:
"-UBARelay" enables the default provider for relay mode.
"-UBAProviders=" overrides the UBA providers ('+' separated list as for UBT) which is otherwise configured through INI files as before.

diff --git a/Engine/Source/Developer/UbaCoordinatorHorde/Private/UbaHordeConfig.cpp b/Engine/Source/Developer/UbaCoordinatorHorde/Private/UbaHordeConfig.cpp
--- a/Engine/Source/Developer/UbaCoordinatorHorde/Private/UbaHordeConfig.cpp
+++ b/Engine/Source/Developer/UbaCoordinatorHorde/Private/UbaHordeConfig.cpp
@@ -45,10 +45,10 @@ struct FUbaHordeConfigParser
 	TParameter<FString> Encryption;
 
 	template <typename T>
-	void SetConfig(TParameter<T>& OutParameter, T InValue)
+	void SetConfig(TParameter<T>& OutParameter, T InValue, const TCHAR* InSource = nullptr)
 	{
 		OutParameter.Value = MoveTemp(InValue);
-		OutParameter.Source = CurrentSource;
+		OutParameter.Source = InSource != nullptr ? InSource : CurrentSource;
 	}
 
 	template <typename T>
@@ -213,9 +213,29 @@ void FUbaHordeConfigParser::ParseConfigFromIniFile()
 	CurrentSource = TEXT("INI");
 
 	// Read array of Horde providers and use default is none is set. Only select the first available since UbaController does not support multiple providers.
-	TArray<FString> HordeProviders;
-	GConfig->GetArray(TEXT("UbaController"), GIsBuildMachine ? TEXT("BuildMachineProviders") : TEXT("Providers"), HordeProviders, GEngineIni);
-	SetConfig(Provider, HordeProviders.Num() > 0 ? HordeProviders[0] : TEXT("Uba.Provider.Horde"));
+	TArray<FString> ProviderList;
+	GConfig->GetArray(TEXT("UbaController"), GIsBuildMachine ? TEXT("BuildMachineProviders") : TEXT("Providers"), ProviderList, GEngineIni);
+
+	// Command line can override the UBA provider from INI settings
+	FString OverrideProviders;
+	if (FParse::Value(FCommandLine::Get(), TEXT("UBAProviders="), OverrideProviders) && OverrideProviders.ParseIntoArray(ProviderList, TEXT("+")) > 0)
+	{
+		SetConfig(Provider, ProviderList[0], TEXT("CMD"));
+	}
+	else if (FParse::Param(FCommandLine::Get(), TEXT("UBARelay")))
+	{
+		// Shortcut arguennt to enable default provider for relay mode
+		SetConfig(Provider, FString(TEXT("Uba.Provider.Horde.Relay")), TEXT("CMD"));
+	}
+	else if (!ProviderList.IsEmpty())
+	{
+		SetConfig(Provider, ProviderList[0]);
+	}
+	else
+	{
+		// Fallback to default provider
+		SetConfig(Provider, FString(TEXT("Uba.Provider.Horde")));
+	}
 
 	GetUbaConfigString(*Provider.Value, TEXT("Enabled"), ProviderEnabled);
 
