[UBA] Adopt specified UBA cluster when it's not "_auto".

This fixes an issue that would only allow the "_auto" and "default" cluster to be specified for UBA providers.

#test Tested with invalid cluster name
#rnx

diff --git a/Engine/Source/Developer/UbaCoordinatorHorde/Private/UbaHordeAgentManager.cpp b/Engine/Source/Developer/UbaCoordinatorHorde/Private/UbaHordeAgentManager.cpp
--- a/Engine/Source/Developer/UbaCoordinatorHorde/Private/UbaHordeAgentManager.cpp
+++ b/Engine/Source/Developer/UbaCoordinatorHorde/Private/UbaHordeAgentManager.cpp
@@ -370,6 +370,11 @@ void FUbaHordeAgentManager::ThreadAgent(FHordeAgentWrapper& Wrapper)
 			ClusterInfo = Future.Get().Value;
 			UE_CLOG(!ClusterInfo.ClusterId.IsEmpty(), LogUbaHorde, Verbose, TEXT("Received resolved cluster ID from Horde: %s"), *ClusterInfo.ClusterId);
 		}
+		else
+		{
+			// Otherwise, just adopt specified cluster identifier
+			ClusterInfo.ClusterId = Config.HordeCluster;
+		}
 
 		// Request Horde machine for UBA agent
 		TSharedPtr<FUbaHordeMetaClient::HordeMachinePromise, ESPMode::ThreadSafe> Promise = HordeMetaClient->RequestMachine(HordeRequestJsonBody, *ClusterInfo.ClusterId);
diff --git a/Engine/Source/Developer/UbaCoordinatorHorde/Private/UbaHordeMetaClient.cpp b/Engine/Source/Developer/UbaCoordinatorHorde/Private/UbaHordeMetaClient.cpp
--- a/Engine/Source/Developer/UbaCoordinatorHorde/Private/UbaHordeMetaClient.cpp
+++ b/Engine/Source/Developer/UbaCoordinatorHorde/Private/UbaHordeMetaClient.cpp
@@ -212,7 +212,8 @@ TSharedPtr<FUbaHordeMetaClient::HordeMachinePromise, ESPMode::ThreadSafe> FUbaHo
 			const EHttpServerResponseCodes ResponseCode = (EHttpServerResponseCodes)HttpResponse->GetResponseCode();
 			const FString ResponseStr = HttpResponse->GetContentAsString();
 
-			if (ResponseCode == EHttpServerResponseCodes::ServiceUnavail ||
+			if (ResponseCode == EHttpServerResponseCodes::NotFound || 
+				ResponseCode == EHttpServerResponseCodes::ServiceUnavail ||
 				ResponseCode == EHttpServerResponseCodes::TooManyRequests)
 			{
 				// Service Unavailable
