[UbaCoordinatorHorde]
* Added so input files to horde bundles exist in order to report issue early instead of getting an exception from the horde agent

#Synced-CL 48787441

diff --git a/Engine/Source/Developer/UbaCoordinatorHorde/Private/UbaHordeAgentManager.cpp b/Engine/Source/Developer/UbaCoordinatorHorde/Private/UbaHordeAgentManager.cpp
--- a/Engine/Source/Developer/UbaCoordinatorHorde/Private/UbaHordeAgentManager.cpp
+++ b/Engine/Source/Developer/UbaCoordinatorHorde/Private/UbaHordeAgentManager.cpp
@@ -263,18 +263,19 @@ void FUbaHordeAgentManager::ThreadAgent(FHordeAgentWrapper& Wrapper)
 			{
 				const TCHAR* Filename;
 				const TCHAR* BundleRef;
+				bool Optional;
 			};
 			const FBundleRec BundleRecs[] =
 			{
 #if PLATFORM_WINDOWS
-				{ TEXT("UbaAgent.exe"), TEXT("UbaAgent.Bundle.ref") },
-				{ TEXT("UbaWine.dll.so"), TEXT("UbaWine.Bundle.ref") },
+				{ TEXT("UbaAgent.exe"), TEXT("UbaAgent.Bundle.ref"), false },
+				{ TEXT("UbaWine.dll.so"), TEXT("UbaWine.Bundle.ref"), true },
 #elif PLATFORM_LINUX
-				{ TEXT("UbaAgent"), TEXT("UbaAgent.Bundle.ref") },
-				{ TEXT("UbaAgent.debug"), TEXT("UbaAgent.debug.Bundle.ref") },
+				{ TEXT("UbaAgent"), TEXT("UbaAgent.Bundle.ref"), false },
+				{ TEXT("UbaAgent.debug"), TEXT("UbaAgent.debug.Bundle.ref"), true },
 //				{ TEXT("libclang_rt.tsan.so"), TEXT("Tsan.Bundle.ref") }, // for debugging
 #elif PLATFORM_MAC
-				{ TEXT("UbaAgent"), TEXT("UbaAgent.Bundle.ref") },
+				{ TEXT("UbaAgent"), TEXT("UbaAgent.Bundle.ref"), false },
 #endif
 			};
 
@@ -283,6 +284,22 @@ void FUbaHordeAgentManager::ThreadAgent(FHordeAgentWrapper& Wrapper)
 				const FString FilePath = FPaths::Combine(BinariesPath, Rec.Filename);
 				FString BundlePath = FPaths::Combine(WorkingDir, Rec.BundleRef);
 
+				if (!FPaths::FileExists(*FilePath))
+				{
+					if (Rec.Optional)
+					{
+						UE_LOG(LogUbaHorde, Warning, TEXT("Failed to find optional %s for Horde bundle"), *FilePath);
+						continue;
+					}
+					else
+					{
+						UE_LOG(LogUbaHorde, Error, TEXT("Failed to create Horde bundle for: %s"), *FilePath);
+						bAskForAgents = false;
+						UpdateStatus(TEXT("Failed to create bundle"));
+						return;
+					}
+				}
+
 				if (!CreateHordeBundleFromFile(*FilePath, *BundlePath))
 				{
 					UE_LOG(LogUbaHorde, Error, TEXT("Failed to create Horde bundle for: %s"), *FilePath);
