PR #13771: [UBA] Add UBA Agent args to match closer to UBT

#Synced-CL 49029036

diff --git a/Engine/Source/Developer/UbaCoordinatorHorde/Private/UbaHordeAgentManager.cpp b/Engine/Source/Developer/UbaCoordinatorHorde/Private/UbaHordeAgentManager.cpp
--- a/Engine/Source/Developer/UbaCoordinatorHorde/Private/UbaHordeAgentManager.cpp
+++ b/Engine/Source/Developer/UbaCoordinatorHorde/Private/UbaHordeAgentManager.cpp
@@ -509,6 +509,11 @@ void FUbaHordeAgentManager::ThreadAgent(FHordeAgentWrapper& Wrapper)
 		UbaAgentArgs.AddRaw("-Eventfile=%UE_HORDE_TERMINATION_SIGNAL_FILE%");
 		UbaAgentArgs.Add(FAnsiString::Printf("-Description=%s", TCHAR_TO_ANSI(*Agent->GetMachineInfo().LeaseLink)));
 
+		UbaAgentArgs.AddRaw("-MaxCpu=%UE_HORDE_CPU_COUNT%");
+		UbaAgentArgs.AddRaw("-MulCpu=%UE_HORDE_CPU_MULTIPLIER%");
+		UbaAgentArgs.AddRaw("-MaxCon=8");
+		UbaAgentArgs.Add(FAnsiString::Printf("-MaxWorkers=%d", Config.HordeMaxWorkers));
+
 		UbaAgentArgs.Finalize();
 
 		// If the machine does not run Windows, enable the compatibility layer Wine to run UbaAgent.exe on POSIX systems
diff --git a/Engine/Source/Developer/UbaCoordinatorHorde/Private/UbaHordeConfig.cpp b/Engine/Source/Developer/UbaCoordinatorHorde/Private/UbaHordeConfig.cpp
--- a/Engine/Source/Developer/UbaCoordinatorHorde/Private/UbaHordeConfig.cpp
+++ b/Engine/Source/Developer/UbaCoordinatorHorde/Private/UbaHordeConfig.cpp
@@ -39,6 +39,7 @@ struct FUbaHordeConfigParser
 	TParameter<FString> Host;
 	TParameter<FString> SentryUrl;
 	TParameter<int32> MaxCores;
+	TParameter<int32> MaxWorkers;
 	TParameter<int32> MaxParallelActions;
 	TParameter<bool> bAllowWine;
 	TParameter<FString> ConnectionMode;
@@ -252,6 +253,7 @@ void FUbaHordeConfigParser::ParseConfigFromIniFile()
 	GetUbaConfigString(*Provider.Value, TEXT("Cluster"), Cluster);
 	GetUbaConfigString(*Provider.Value, TEXT("LocalHost"), Host);
 	GetUbaConfigInt(*Provider.Value, TEXT("MaxCores"), MaxCores);
+	GetUbaConfigInt(*Provider.Value, TEXT("MaxWorkers"), MaxWorkers);
 	GetUbaConfigBool(*Provider.Value, TEXT("AllowWine"), bAllowWine);
 	GetUbaConfigString(*Provider.Value, TEXT("ConnectionMode"), ConnectionMode);
 	GetUbaConfigString(*Provider.Value, TEXT("Encryption"), Encryption);
@@ -267,6 +269,7 @@ void FUbaHordeConfigParser::ParseConfigFromCommandline(const TCHAR* CmdLine)
 	ParseCmdLineString(CmdLine, TEXT("UBAHordePool="), Pool);
 	ParseCmdLineBoolean(CmdLine, TEXT("UBAHordeAllowWine"), bAllowWine);
 	ParseCmdLineInt(CmdLine, TEXT("UBAHordeMaxCores="), MaxCores);
+	ParseCmdLineInt(CmdLine, TEXT("UBAHordeMaxWorkers="), MaxWorkers);
 	ParseCmdLineString(CmdLine, TEXT("UBAHordeHost="), Host);
 	ParseCmdLineString(CmdLine, TEXT("UBAHordeCluster="), Cluster);
 	ParseCmdLineString(CmdLine, TEXT("UBAHordeRequirements="), Condition);
@@ -384,6 +387,7 @@ void FUbaHordeConfig::Initialize()
 	ConfigParser.GetConfig(ConfigParser.Host, HordeHost);
 	ConfigParser.GetConfig(ConfigParser.SentryUrl, UbaSentryUrl);
 	ConfigParser.GetConfig(ConfigParser.MaxCores, HordeMaxCores);
+	ConfigParser.GetConfig(ConfigParser.MaxWorkers, HordeMaxWorkers);
 	ConfigParser.GetConfig(ConfigParser.MaxParallelActions, MaxParallelActions);
 	ConfigParser.GetConfig(ConfigParser.bAllowWine, bHordeAllowWine);
 
@@ -426,6 +430,7 @@ void FUbaHordeConfig::Initialize()
 			TEXT("\n  Cluster:           %s (%s)")
 			TEXT("\n  Host:              %s (%s)")
 			TEXT("\n  MaxCores:          %d (%s)")
+			TEXT("\n  MaxWorkers:        %d (%s)")
 			TEXT("\n  Allow Wine:        %s (%s)")
 			TEXT("\n  Mode:              %s (%s)")
 			TEXT("\n  Encryption:        %s (%s)")
@@ -436,6 +441,7 @@ void FUbaHordeConfig::Initialize()
 			, GetConfigValue(*HordeCluster), GetConfigSource(ConfigParser.Cluster.Source)
 			, GetConfigValue(*HordeHost), GetConfigSource(ConfigParser.Host.Source)
 			, HordeMaxCores, GetConfigSource(ConfigParser.MaxCores.Source)
+			, HordeMaxWorkers, GetConfigSource(ConfigParser.MaxWorkers.Source)
 			, bHordeAllowWine ? TEXT("Yes") : TEXT("No"), GetConfigSource(ConfigParser.bAllowWine.Source)
 			, GetConfigValue(LexToString(HordeConnectionMode)), GetConfigSource(ConfigParser.ConnectionMode.Source)
 			, GetConfigValue(LexToString(HordeEncryption)), GetConfigSource(ConfigParser.Encryption.Source)
diff --git a/Engine/Source/Developer/UbaCoordinatorHorde/Public/UbaHordeConfig.h b/Engine/Source/Developer/UbaCoordinatorHorde/Public/UbaHordeConfig.h
--- a/Engine/Source/Developer/UbaCoordinatorHorde/Public/UbaHordeConfig.h
+++ b/Engine/Source/Developer/UbaCoordinatorHorde/Public/UbaHordeConfig.h
@@ -65,6 +65,9 @@ public:
 	/** Maxmium number of CPU cores allowed to be used by build session. */
 	int32 HordeMaxCores = 2048;
 
+	/** Maximum workers allowed to be used by build session. */
+	int32 HordeMaxWorkers = 64;
+
 	/** Maximum number of local CPU cores allowed to be used by build session. By default -1, which indicates to use as many CPU cores as the local machine can provide. */
 	int32 MaxParallelActions = -1;
 
