[UbaCoordinatorHorde]
* Added code to be able to request cache server

#Synced-CL 52497433

diff --git a/Engine/Source/Developer/UbaCoordinatorHorde/Private/UbaHordeAgentManager.cpp b/Engine/Source/Developer/UbaCoordinatorHorde/Private/UbaHordeAgentManager.cpp
--- a/Engine/Source/Developer/UbaCoordinatorHorde/Private/UbaHordeAgentManager.cpp
+++ b/Engine/Source/Developer/UbaCoordinatorHorde/Private/UbaHordeAgentManager.cpp
@@ -56,6 +56,10 @@ FUbaHordeAgentManager::~FUbaHordeAgentManager()
 	for (TUniquePtr<FHordeAgentWrapper>& Agent : Agents)
 	{
 		Agent->ShouldExit->Trigger();
+	}
+
+	for (TUniquePtr<FHordeAgentWrapper>& Agent : Agents)
+	{
 		Agent->Thread.Join();
 		FGenericPlatformProcess::ReturnSynchEventToPool(Agent->ShouldExit);
 	}
@@ -112,6 +116,73 @@ uint32 FUbaHordeAgentManager::GetActiveCoreCount() const
 	return ActiveCoreCount;
 }
 
+bool FUbaHordeAgentManager::RequestCacheServer(FString& OutCacheEndpoint, FString& OutCacheSessionKey, bool& bOutWriteAccess)
+{
+	FScopeLock ScopeLock(&BundleRefPathsLock);
+	if (!HordeMetaClient)
+	{
+		HordeMetaClient = MakeUnique<FUbaHordeMetaClient>();
+		if (!HordeMetaClient->RefreshHttpClient())
+		{
+			UE_LOG(LogUbaHorde, Warning, TEXT("Failed to create HttpClient for cache server request"));
+			HordeMetaClient.Reset();
+			return false;
+		}
+	}
+	ScopeLock.Unlock();
+
+	const FUbaHordeConfig& Config = FUbaHordeConfig::Get();
+
+	// Resolve cluster ID if set to auto
+	FHordeClusterInfo ClusterInfo;
+	if (Config.HordeCluster == FUbaHordeConfig::ClusterAuto)
+	{
+		const FString HordeRequestJsonBody = FUbaHordeMetaClient::BuildHordeRequestJsonBody(
+			Config.HordeConnectionMode == EUbaHordeConnectionMode::Direct ? Config.GetHordePool() : TEXT(""),
+			Config.HordeConnectionMode,
+			Config.HordeEncryption,
+			*Config.HordeCondition,
+			false, // not exclusive for cache lookup
+			Config.bHordeAllowWine
+		);
+		TSharedPtr<FUbaHordeMetaClient::HordeClusterPromise, ESPMode::ThreadSafe> ClusterPromise = HordeMetaClient->RequestClusterId(HordeRequestJsonBody);
+		if (!ClusterPromise)
+		{
+			UE_LOG(LogUbaHorde, Warning, TEXT("Failed to resolve Horde cluster ID for cache server request"));
+			return false;
+		}
+		TFuture<TTuple<FHttpResponsePtr, FHordeClusterInfo>> ClusterFuture = ClusterPromise->GetFuture();
+		ClusterFuture.Wait();
+		ClusterInfo = ClusterFuture.Get().Value;
+	}
+	else
+	{
+		ClusterInfo.ClusterId = Config.HordeCluster;
+	}
+
+	TSharedPtr<FUbaHordeMetaClient::HordeCacheServerPromise, ESPMode::ThreadSafe> Promise = HordeMetaClient->RequestCacheServer(*ClusterInfo.ClusterId);
+	if (!Promise)
+	{
+		UE_LOG(LogUbaHorde, Warning, TEXT("Failed to request Horde cache server from cluster: %s"), *ClusterInfo.ClusterId);
+		return false;
+	}
+
+	TFuture<TTuple<FHttpResponsePtr, FHordeCacheServerInfo>> Future = Promise->GetFuture();
+	Future.Wait();
+	FHordeCacheServerInfo CacheInfo = Future.Get().Value;
+
+	if (CacheInfo.CacheEndpoint.IsEmpty())
+	{
+		UE_LOG(LogUbaHorde, Display, TEXT("No cache server available from Horde cluster: %s"), *ClusterInfo.ClusterId);
+		return false;
+	}
+
+	OutCacheEndpoint = CacheInfo.CacheEndpoint;
+	OutCacheSessionKey = CacheInfo.CacheSessionKey;
+	bOutWriteAccess = CacheInfo.bWriteAccess;
+	return true;
+}
+
 // Creates a bundle blob (one of several chunks of a file) to be uploaded to Horde
 // This code has been adopted from the HordeTest project.
 // See 'Engine/Source/Programs/Horde/Samples/HordeTest/Main.cpp'.
diff --git a/Engine/Source/Developer/UbaCoordinatorHorde/Private/UbaHordeMetaClient.cpp b/Engine/Source/Developer/UbaCoordinatorHorde/Private/UbaHordeMetaClient.cpp
--- a/Engine/Source/Developer/UbaCoordinatorHorde/Private/UbaHordeMetaClient.cpp
+++ b/Engine/Source/Developer/UbaCoordinatorHorde/Private/UbaHordeMetaClient.cpp
@@ -374,3 +374,86 @@ TSharedPtr<FUbaHordeMetaClient::HordeMachinePromise, ESPMode::ThreadSafe> FUbaHo
 
 	return Promise;
 }
+
+TSharedPtr<FUbaHordeMetaClient::HordeCacheServerPromise, ESPMode::ThreadSafe> FUbaHordeMetaClient::RequestCacheServer(const TCHAR* ClusterId)
+{
+	TSharedPtr<HordeCacheServerPromise> Promise = MakeShared<HordeCacheServerPromise, ESPMode::ThreadSafe>();
+
+	if (bClientNeedsRefresh)
+	{
+		RefreshHttpClient();
+	}
+
+	const FString ResourcePath = FString::Format(TEXT("api/v2/compute/{0}/uba-cache"), { ClusterId != nullptr && *ClusterId != TEXT('\0') ? ClusterId : TEXT("default") });
+	FHttpRequestRef Request = HttpClient->CreateRequest(TEXT("POST"), *ResourcePath);
+
+	Request->SetDelegateThreadPolicy(EHttpRequestDelegateThreadPolicy::CompleteOnHttpThread);
+	Request->SetContentAsString(TEXT("{}"));
+	Request->SetHeader(TEXT("Content-Type"), TEXT("application/json"));
+
+	Request->OnProcessRequestComplete().BindLambda(
+		[Promise](FHttpRequestPtr /*Request*/, FHttpResponsePtr HttpResponse, bool bConnectedSuccessfully)
+		{
+			FHordeCacheServerInfo Info;
+
+			ON_SCOPE_EXIT
+			{
+				Promise->SetValue(MakeTuple(HttpResponse, Info));
+			};
+
+			if (!bConnectedSuccessfully || !HttpResponse.IsValid())
+			{
+				UE_LOG(LogUbaHorde, Display, TEXT("No response from Horde for cache server request"));
+				return;
+			}
+
+			const EHttpServerResponseCodes ResponseCode = (EHttpServerResponseCodes)HttpResponse->GetResponseCode();
+			const FString ResponseStr = HttpResponse->GetContentAsString();
+
+			if (ResponseCode == EHttpServerResponseCodes::ServiceUnavail || ResponseCode == EHttpServerResponseCodes::TooManyRequests)
+			{
+				UE_LOG(LogUbaHorde, Display, TEXT("Horde cache server request returned with HTTP/%d: %s"), (int32)ResponseCode, *ResponseStr);
+				return;
+			}
+
+			if (ResponseCode == EHttpServerResponseCodes::Denied)
+			{
+				UE_LOG(LogUbaHorde, Display, TEXT("Token expired during cache server request"));
+				return;
+			}
+
+			TSharedPtr<FJsonValue> OutJson;
+			TSharedRef<TJsonReader<>> Reader = TJsonReaderFactory<>::Create(ResponseStr);
+
+			if (!FJsonSerializer::Deserialize(Reader, OutJson, FJsonSerializer::EFlags::None))
+			{
+				UE_LOG(LogUbaHorde, Display, TEXT("Invalid response body for cache server request (HTTP/%d): %s"), (int32)ResponseCode, *ResponseStr);
+				return;
+			}
+
+			TSharedPtr<FJsonValue> CacheEndpointValue = OutJson->AsObject()->TryGetField(TEXT("cacheEndpoint"));
+			if (!CacheEndpointValue.Get())
+			{
+				UE_LOG(LogUbaHorde, Display, TEXT("Missing \"cacheEndpoint\" in cache server response (HTTP/%d): %s"), (int32)ResponseCode, *ResponseStr);
+				return;
+			}
+
+			Info.CacheEndpoint = CacheEndpointValue->AsString();
+
+			if (TSharedPtr<FJsonValue> KeyValue = OutJson->AsObject()->TryGetField(TEXT("cacheSessionKey")))
+			{
+				Info.CacheSessionKey = KeyValue->AsString();
+			}
+
+			if (TSharedPtr<FJsonValue> WriteAccessValue = OutJson->AsObject()->TryGetField(TEXT("writeAccess")))
+			{
+				Info.bWriteAccess = WriteAccessValue->AsBool();
+			}
+
+			UE_LOG(LogUbaHorde, Display, TEXT("Horde cache server assigned: %s (WriteAccess: %s)"), *Info.CacheEndpoint, Info.bWriteAccess ? TEXT("true") : TEXT("false"));
+		});
+
+	Request->ProcessRequest();
+
+	return Promise;
+}
diff --git a/Engine/Source/Developer/UbaCoordinatorHorde/Private/UbaHordeMetaClient.h b/Engine/Source/Developer/UbaCoordinatorHorde/Private/UbaHordeMetaClient.h
--- a/Engine/Source/Developer/UbaCoordinatorHorde/Private/UbaHordeMetaClient.h
+++ b/Engine/Source/Developer/UbaCoordinatorHorde/Private/UbaHordeMetaClient.h
@@ -69,6 +69,13 @@ struct FHordeClusterInfo
 	FString ClusterId = TEXT("default");
 };
 
+struct FHordeCacheServerInfo
+{
+	FString CacheEndpoint;    // "ip:port" format, empty if unavailable
+	FString CacheSessionKey;  // AES encryption key for cache communication
+	bool bWriteAccess = false;
+};
+
 // This encapsulates the mechanism of talking to the Horde "meta server". The server which
 // can grant us access to remote machines.
 class FUbaHordeMetaClient final
@@ -77,6 +84,7 @@ public:
 	// We return the HttpResponse in case more information could be used out of it later.
 	using HordeMachinePromise = TPromise<TTuple<FHttpResponsePtr, FHordeRemoteMachineInfo>>;
 	using HordeClusterPromise = TPromise<TTuple<FHttpResponsePtr, FHordeClusterInfo>>;
+	using HordeCacheServerPromise = TPromise<TTuple<FHttpResponsePtr, FHordeCacheServerInfo>>;
 
 	FUbaHordeMetaClient() = default;
 
@@ -96,6 +104,8 @@ public:
 	// FHordeRemoteMachineInfo MachineInfo = Future.Get().Value;
 	TSharedPtr<HordeMachinePromise, ESPMode::ThreadSafe> RequestMachine(const FString& HordeRequestJsonBody, const TCHAR* ClusterId = nullptr);
 
+	TSharedPtr<HordeCacheServerPromise, ESPMode::ThreadSafe> RequestCacheServer(const TCHAR* ClusterId = nullptr);
+
 private:
 	FString ServerUrl;
 	TUniquePtr<FHordeHttpClient> HttpClient;
diff --git a/Engine/Source/Developer/UbaCoordinatorHorde/Public/UbaHordeAgentManager.h b/Engine/Source/Developer/UbaCoordinatorHorde/Public/UbaHordeAgentManager.h
--- a/Engine/Source/Developer/UbaCoordinatorHorde/Public/UbaHordeAgentManager.h
+++ b/Engine/Source/Developer/UbaCoordinatorHorde/Public/UbaHordeAgentManager.h
@@ -31,6 +31,10 @@ public:
 	// Returns the active number of cores allocated across all agents.
 	UBACOORDINATORHORDE_API uint32 GetActiveCoreCount() const;
 
+	// Requests a UBA cache server from Horde. Returns true if a cache server was assigned.
+	// OutCacheEndpoint will be in "ip:port" format. OutCacheSessionKey is the AES encryption key.
+	UBACOORDINATORHORDE_API bool RequestCacheServer(FString& OutCacheEndpoint, FString& OutCacheSessionKey, bool& bOutWriteAccess);
+
 private:
 	struct FHordeAgentWrapper
 	{
